Fix for nounity error

diff --git a/Engine/Source/Editor/LevelInstanceEditor/Private/LevelInstanceEditorModule.cpp b/Engine/Source/Editor/LevelInstanceEditor/Private/LevelInstanceEditorModule.cpp
--- a/Engine/Source/Editor/LevelInstanceEditor/Private/LevelInstanceEditorModule.cpp
+++ b/Engine/Source/Editor/LevelInstanceEditor/Private/LevelInstanceEditorModule.cpp
@@ -27,6 +27,7 @@
 #include "Kismet2/BlueprintEditorUtils.h"
 #include "Misc/ScopeExit.h"
 #include "Widgets/Input/SNumericEntryBox.h"
+#include "Widgets/Input/SButton.h"
 
 IMPLEMENT_MODULE( FLevelInstanceEditorModule, LevelInstanceEditor );
 
