Fix LevelInstance DetailsCustomization Registration

diff --git a/Engine/Source/Editor/LevelInstanceEditor/Private/LevelInstanceEditorModule.cpp b/Engine/Source/Editor/LevelInstanceEditor/Private/LevelInstanceEditorModule.cpp
--- a/Engine/Source/Editor/LevelInstanceEditor/Private/LevelInstanceEditorModule.cpp
+++ b/Engine/Source/Editor/LevelInstanceEditor/Private/LevelInstanceEditorModule.cpp
@@ -548,7 +548,7 @@ void FLevelInstanceEditorModule::StartupModule()
 	ExtendContextMenu();
 
 	FPropertyEditorModule& PropertyModule = FModuleManager::LoadModuleChecked<FPropertyEditorModule>("PropertyEditor");
-	PropertyModule.RegisterCustomClassLayout("LevelInstanceActor", FOnGetDetailCustomizationInstance::CreateStatic(&FLevelInstanceActorDetails::MakeInstance));
+	PropertyModule.RegisterCustomClassLayout("LevelInstance", FOnGetDetailCustomizationInstance::CreateStatic(&FLevelInstanceActorDetails::MakeInstance));
 	PropertyModule.NotifyCustomizationModuleChanged();
 
 	// GEditor needs to be set before this module is loaded
