Added popup warning when breaking a Level Instance

diff --git a/Engine/Source/Editor/LevelInstanceEditor/Private/LevelInstanceEditorModule.cpp b/Engine/Source/Editor/LevelInstanceEditor/Private/LevelInstanceEditorModule.cpp
--- a/Engine/Source/Editor/LevelInstanceEditor/Private/LevelInstanceEditorModule.cpp
+++ b/Engine/Source/Editor/LevelInstanceEditor/Private/LevelInstanceEditorModule.cpp
@@ -428,7 +428,15 @@ namespace LevelInstanceMenuUtils
 					SNew(SButton)
 					.HAlign(HAlign_Center)
 					.ContentPadding(FEditorStyle::GetMargin("StandardDialog.ContentPadding"))
-					.OnClicked_Lambda([ContextLevelInstance, LevelInstanceSubsystem]() {LevelInstanceSubsystem->BreakLevelInstance(ContextLevelInstance, BreakLevels); return FReply::Handled(); })
+					.OnClicked_Lambda([ContextLevelInstance, LevelInstanceSubsystem]() 
+						{
+							const FText LevelInstanceBreakWarning = LOCTEXT("BreakingLevelInstance", "You are about to break the level instance this action cannot be undone. Are you sure ?");
+							if (FMessageDialog::Open(EAppMsgType::YesNo, LevelInstanceBreakWarning) == EAppReturnType::Yes)
+							{
+								LevelInstanceSubsystem->BreakLevelInstance(ContextLevelInstance, BreakLevels);
+							}
+							return FReply::Handled();
+						})
 					.Text(LOCTEXT("BreakLevelInstances_BreakLevelInstanceButton", "Break Level Instance"))
 				];
 
