Fixed the grammar issue when showing the popup warning.

diff --git a/Engine/Source/Editor/LevelInstanceEditor/Private/LevelInstanceEditorModule.cpp b/Engine/Source/Editor/LevelInstanceEditor/Private/LevelInstanceEditorModule.cpp
--- a/Engine/Source/Editor/LevelInstanceEditor/Private/LevelInstanceEditorModule.cpp
+++ b/Engine/Source/Editor/LevelInstanceEditor/Private/LevelInstanceEditorModule.cpp
@@ -430,7 +430,7 @@ namespace LevelInstanceMenuUtils
 					.ContentPadding(FEditorStyle::GetMargin("StandardDialog.ContentPadding"))
 					.OnClicked_Lambda([ContextLevelInstance, LevelInstanceSubsystem]() 
 						{
-							const FText LevelInstanceBreakWarning = LOCTEXT("BreakingLevelInstance", "You are about to break the level instance this action cannot be undone. Are you sure ?");
+							const FText LevelInstanceBreakWarning = LOCTEXT("BreakingLevelInstance", "You are about to break the level instance. This action cannot be undone. Are you sure ?");
 							if (FMessageDialog::Open(EAppMsgType::YesNo, LevelInstanceBreakWarning) == EAppReturnType::Yes)
 							{
 								LevelInstanceSubsystem->BreakLevelInstance(ContextLevelInstance, BreakLevels);
