LevelInstance: Added an actor which allows the user to change the pivot point of the level instance.

diff --git a/Engine/Source/Editor/LevelInstanceEditor/LevelInstanceEditor.Build.cs b/Engine/Source/Editor/LevelInstanceEditor/LevelInstanceEditor.Build.cs
--- a/Engine/Source/Editor/LevelInstanceEditor/LevelInstanceEditor.Build.cs
+++ b/Engine/Source/Editor/LevelInstanceEditor/LevelInstanceEditor.Build.cs
@@ -32,7 +32,8 @@ public class LevelInstanceEditor : ModuleRules
 				"ContentBrowser",
 				"AssetTools",
 				"ClassViewer",
-				"MessageLog"
+				"MessageLog",
+				"EditorWidgets"
 			}
 		);
 
diff --git a/Engine/Source/Editor/LevelInstanceEditor/Private/LevelInstanceEditorModule.cpp b/Engine/Source/Editor/LevelInstanceEditor/Private/LevelInstanceEditorModule.cpp
--- a/Engine/Source/Editor/LevelInstanceEditor/Private/LevelInstanceEditorModule.cpp
+++ b/Engine/Source/Editor/LevelInstanceEditor/Private/LevelInstanceEditorModule.cpp
@@ -1,6 +1,7 @@
 // Copyright Epic Games, Inc. All Rights Reserved.
 #include "LevelInstanceEditorModule.h"
 #include "LevelInstanceActorDetails.h"
+#include "LevelInstancePivotDetails.h"
 #include "LevelInstance/LevelInstanceSubsystem.h"
 #include "LevelInstance/LevelInstanceActor.h"
 #include "LevelInstance/Packed/PackedLevelInstanceBuilder.h"
@@ -30,6 +31,8 @@
 #include "Misc/ScopeExit.h"
 #include "Widgets/Input/SNumericEntryBox.h"
 #include "Widgets/Input/SButton.h"
+#include "Widgets/SWindow.h"
+#include "SNewLevelInstanceDialog.h"
 #include "MessageLogModule.h"
 
 IMPLEMENT_MODULE( FLevelInstanceEditorModule, LevelInstanceEditor );
@@ -298,7 +301,7 @@ namespace LevelInstanceMenuUtils
 		}
 	}
 
-	void CreateLevelInstanceFromSelection(ULevelInstanceSubsystem* LevelInstanceSubsystem, ELevelInstanceCreationType CreationType)
+	void CreateLevelInstanceFromSelection(ULevelInstanceSubsystem* LevelInstanceSubsystem)
 	{
 		TArray<AActor*> ActorsToMove;
 		ActorsToMove.Reserve(GEditor->GetSelectedActorCount());
@@ -311,54 +314,69 @@ namespace LevelInstanceMenuUtils
 		}
 
 		IMainFrameModule& MainFrameModule = FModuleManager::GetModuleChecked<IMainFrameModule>("MainFrame");
-		FNewLevelDialogModule& NewLevelDialogModule = FModuleManager::LoadModuleChecked<FNewLevelDialogModule>("NewLevelDialog");
-		FString TemplateMapPackage;
-		if (!GetMutableDefault<ULevelInstanceEditorSettings>()->TemplateMapInfos.Num() || NewLevelDialogModule.CreateAndShowTemplateDialog(MainFrameModule.GetParentWindow(), LOCTEXT("LevelInstanceTemplateDialog", "Choose Level Instance Template..."), GetMutableDefault<ULevelInstanceEditorSettings>()->TemplateMapInfos, TemplateMapPackage))
-		{
-			UPackage* TemplatePackage = !TemplateMapPackage.IsEmpty() ? LoadPackage(nullptr, *TemplateMapPackage, LOAD_None) : nullptr;
-			UWorld* TemplateWorld = TemplatePackage ? UWorld::FindWorldInPackage(TemplatePackage) : nullptr;
 
-			if (!LevelInstanceSubsystem->CreateLevelInstanceFrom(ActorsToMove, CreationType, TemplateWorld))
-			{
-				FText Title = LOCTEXT("CreateFromSelectionFailTitle", "Create from selection failed");
-				FMessageDialog::Open(EAppMsgType::Ok, LOCTEXT("CreateFromSelectionFailMsg", "Failed to create LevelInstance from selection. Check log for details."), &Title);
-			}
-		}
-	}
+		TSharedPtr<SWindow> NewLevelInstanceWindow =
+			SNew(SWindow)
+			.Title(LOCTEXT("NewLevelInstanceWindowTitle", "New Level Instance"))
+			.ClientSize(SNewLevelInstanceDialog::DEFAULT_WINDOW_SIZE)
+			.SizingRule(ESizingRule::UserSized)
+			.SupportsMinimize(false)
+			.SupportsMaximize(false)
+			.SizingRule(ESizingRule::FixedSize);
 
-	void CreateLevelInstanceFromSelectionSubMenu(UToolMenu* ToolMenu, ULevelInstanceSubsystem* LevelInstanceSubsystem)
-	{
-		FToolMenuSection* Section = &ToolMenu->AddSection(FName("LevelInstanceCreateFromSelectionSub"));
-		UEnum* CreationTypeEnum = StaticEnum<ELevelInstanceCreationType>();
-		check(CreationTypeEnum);
-		for (int32 i = 0; i < CreationTypeEnum->NumEnums()-1; ++i)
+		TSharedRef<SNewLevelInstanceDialog> NewLevelInstanceDialog =
+			SNew(SNewLevelInstanceDialog)
+			.ParentWindow(NewLevelInstanceWindow)
+			.PivotActors(ActorsToMove);
+
+		NewLevelInstanceWindow->SetContent(NewLevelInstanceDialog);
+
+		FSlateApplication::Get().AddModalWindow(NewLevelInstanceWindow.ToSharedRef(), MainFrameModule.GetParentWindow());
+
+
+		if (NewLevelInstanceDialog->ClickedOk())
 		{
-			if (!CreationTypeEnum->HasMetaData(TEXT("Hidden"), i))
+			ELevelInstanceCreationType CreationType = NewLevelInstanceDialog->GetCreationType();
+			ELevelInstancePivotType PivotType = NewLevelInstanceDialog->GetPivotType();
+			AActor* PivotActor = NewLevelInstanceDialog->GetPivotActor();
+
+			FNewLevelDialogModule& NewLevelDialogModule = FModuleManager::LoadModuleChecked<FNewLevelDialogModule>("NewLevelDialog");
+			FString TemplateMapPackage;
+			if (!GetMutableDefault<ULevelInstanceEditorSettings>()->TemplateMapInfos.Num() || NewLevelDialogModule.CreateAndShowTemplateDialog(MainFrameModule.GetParentWindow(), LOCTEXT("LevelInstanceTemplateDialog", "Choose Level Instance Template..."), GetMutableDefault<ULevelInstanceEditorSettings>()->TemplateMapInfos, TemplateMapPackage))
 			{
-				ELevelInstanceCreationType CreationType = (ELevelInstanceCreationType)CreationTypeEnum->GetValueByIndex(i);
-				FToolUIAction CreateFromSelectionAction;
-				CreateFromSelectionAction.ExecuteAction.BindLambda([LevelInstanceSubsystem, CreationType](const FToolMenuContext&)
-				{
-					CreateLevelInstanceFromSelection(LevelInstanceSubsystem, CreationType);
-				});
+				UPackage* TemplatePackage = !TemplateMapPackage.IsEmpty() ? LoadPackage(nullptr, *TemplateMapPackage, LOAD_None) : nullptr;
+				UWorld* TemplateWorld = TemplatePackage ? UWorld::FindWorldInPackage(TemplatePackage) : nullptr;
 
-				Section->AddMenuEntry(NAME_None, CreationTypeEnum->GetDisplayNameTextByIndex(i), TAttribute<FText>(), FSlateIcon(), CreateFromSelectionAction);
+				if (!LevelInstanceSubsystem->CreateLevelInstanceFrom(ActorsToMove, CreationType, PivotType, PivotActor, TemplateWorld))
+				{
+					FText Title = LOCTEXT("CreateFromSelectionFailTitle", "Create from selection failed");
+					FMessageDialog::Open(EAppMsgType::Ok, LOCTEXT("CreateFromSelectionFailMsg", "Failed to create LevelInstance from selection. Check log for details."), &Title);
+				}
 			}
 		}
 	}
-
+	
 	void CreateCreateMenu(UToolMenu* Menu)
 	{
 		if (ULevelInstanceSubsystem* LevelInstanceSubsystem = GEditor->GetEditorWorldContext().World()->GetSubsystem<ULevelInstanceSubsystem>())
 		{
 			FToolMenuSection& Section = CreateLevelInstanceSection(Menu);
-			
-			Section.AddSubMenu(
+			FToolUIAction UIAction;
+			UIAction.ExecuteAction.BindLambda([LevelInstanceSubsystem](const FToolMenuContext& MenuContext)
+			{
+				CreateLevelInstanceFromSelection(LevelInstanceSubsystem);
+			});
+			UIAction.CanExecuteAction.BindLambda([](const FToolMenuContext& MenuContext)
+			{
+				return GEditor->GetSelectedActorCount() > 0;
+			});
+
+			Section.AddMenuEntry(
 				"CreateLevelInstanceFromSelection",
-				LOCTEXT("CreateLevelInstanceFromSelection", "Create from selection"),
+				LOCTEXT("CreateLevelInstanceFromSelection", "Create from selection..."),
 				TAttribute<FText>(),
-				FNewToolMenuDelegate::CreateStatic(&CreateLevelInstanceFromSelectionSubMenu, LevelInstanceSubsystem)
-			);
+				TAttribute<FSlateIcon>(),
+				UIAction);
 		}
 	}
 
@@ -642,6 +660,7 @@ void FLevelInstanceEditorModule::StartupModule()
 
 	FPropertyEditorModule& PropertyModule = FModuleManager::LoadModuleChecked<FPropertyEditorModule>("PropertyEditor");
 	PropertyModule.RegisterCustomClassLayout("LevelInstance", FOnGetDetailCustomizationInstance::CreateStatic(&FLevelInstanceActorDetails::MakeInstance));
+	PropertyModule.RegisterCustomClassLayout("LevelInstancePivot", FOnGetDetailCustomizationInstance::CreateStatic(&FLevelInstancePivotDetails::MakeInstance));
 	PropertyModule.NotifyCustomizationModuleChanged();
 
 	// GEditor needs to be set before this module is loaded
diff --git a/Engine/Source/Editor/LevelInstanceEditor/Private/LevelInstancePivotDetails.cpp b/Engine/Source/Editor/LevelInstanceEditor/Private/LevelInstancePivotDetails.cpp
new file mode 100644
--- /dev/null
+++ b/Engine/Source/Editor/LevelInstanceEditor/Private/LevelInstancePivotDetails.cpp
@@ -0,0 +1,158 @@
+// Copyright Epic Games, Inc. All Rights Reserved.
+
+#include "LevelInstancePivotDetails.h"
+#include "DetailLayoutBuilder.h"
+#include "DetailCategoryBuilder.h"
+#include "DetailWidgetRow.h"
+#include "Widgets/Input/SButton.h"
+#include "Widgets/Input/SMultiLineEditableTextBox.h"
+#include "SEnumCombo.h"
+#include "ScopedTransaction.h"
+
+#include "Engine/World.h"
+#include "LevelInstance/LevelInstanceActor.h"
+#include "LevelInstance/LevelInstanceEditorPivotActor.h"
+
+#define LOCTEXT_NAMESPACE "FLevelInstancePivotDetails"
+
+FLevelInstancePivotDetails::FLevelInstancePivotDetails()
+{
+}
+
+TSharedRef<IDetailCustomization> FLevelInstancePivotDetails::MakeInstance()
+{
+	return MakeShareable(new FLevelInstancePivotDetails);
+}
+
+void FLevelInstancePivotDetails::CustomizeDetails(IDetailLayoutBuilder& DetailBuilder)
+{
+	TArray<TWeakObjectPtr<UObject>> EditingObjects;
+	DetailBuilder.GetObjectsBeingCustomized(EditingObjects);
+
+	if (EditingObjects.Num() > 1)
+	{
+		return;
+	}
+
+	TWeakObjectPtr<ALevelInstancePivot> EditingObject = Cast<ALevelInstancePivot>(EditingObjects[0].Get());
+
+	UWorld* World = EditingObject->GetWorld();
+
+	if (!World)
+	{
+		return;
+	}
+
+	PivotActor = nullptr;
+	PivotType = ELevelInstancePivotType::CenterMinZ;
+	PivotActors.Reset();
+	for (AActor* Actor : EditingObject->GetLevel()->Actors)
+	{
+		if (Actor != nullptr && Actor != EditingObject && Actor->GetRootComponent() != nullptr && !Actor->IsA<ABrush>())
+		{
+			PivotActors.Add(MakeShareable(new FActorInfo{ Actor }));
+		}
+	}
+
+	IDetailCategoryBuilder& ChangePivotCategory = DetailBuilder.EditCategory("Change Pivot", FText::GetEmpty(), ECategoryPriority::Important);
+		
+	ChangePivotCategory.AddCustomRow(LOCTEXT("PivotTypeRow", "Type"), false)
+		.NameContent()
+		[
+			SNew(STextBlock)
+			.Text(LOCTEXT("PivotTypeText", "Type"))
+		]
+		.ValueContent()
+		[
+			SNew(SEnumComboBox, StaticEnum<ELevelInstancePivotType>())
+			.ContentPadding(FCoreStyle::Get().GetMargin("StandardDialog.ContentPadding"))
+			.CurrentValue(this, &FLevelInstancePivotDetails::GetSelectedPivotType)
+			.OnEnumSelectionChanged(this, &FLevelInstancePivotDetails::OnSelectedPivotTypeChanged)
+		];
+
+	ChangePivotCategory.AddCustomRow(LOCTEXT("PivotActorRow", "Actor"), false)
+		.NameContent()
+		[
+			SNew(STextBlock)
+			.Text(LOCTEXT("PivotActorText", "Actor"))
+		]
+		.ValueContent()
+		[
+			SNew(SComboBox<TSharedPtr<FActorInfo>>)
+			.ContentPadding(FCoreStyle::Get().GetMargin("StandardDialog.ContentPadding"))
+			.OptionsSource(&PivotActors)
+			.OnGenerateWidget(this, &FLevelInstancePivotDetails::OnGeneratePivotActorWidget)
+			.OnSelectionChanged(this, &FLevelInstancePivotDetails::OnSelectedPivotActorChanged)
+			.IsEnabled(this, &FLevelInstancePivotDetails::IsPivotActorSelectionEnabled)
+			[
+				SNew(STextBlock)
+				.Text(this, &FLevelInstancePivotDetails::GetSelectedPivotActorText)
+			]
+		];
+			
+	ChangePivotCategory.AddCustomRow(LOCTEXT("ApplyRow", "Apply"), false)
+		[
+			SNew(SHorizontalBox)
+			+ SHorizontalBox::Slot()
+			.FillWidth(1)
+			[
+				SNew(SButton)
+				.IsEnabled(this, &FLevelInstancePivotDetails::IsApplyButtonEnabled, EditingObject)
+				.Text(LOCTEXT("ChangePivotButton", "Apply"))
+				.HAlign(HAlign_Center)
+				.OnClicked(this, &FLevelInstancePivotDetails::OnApplyButtonClicked, EditingObject)
+			]
+		];
+}
+
+bool FLevelInstancePivotDetails::IsApplyButtonEnabled(TWeakObjectPtr<ALevelInstancePivot> LevelInstancePivot) const
+{
+	return PivotType != ELevelInstancePivotType::Actor || PivotActor.IsValid();
+}
+
+FReply FLevelInstancePivotDetails::OnApplyButtonClicked(TWeakObjectPtr<ALevelInstancePivot> LevelInstancePivot)
+{
+	if (LevelInstancePivot.IsValid())
+	{
+		FScopedTransaction Transaction(LOCTEXT("ChangePivot", "Change Level Instance Pivot"));
+		LevelInstancePivot->SetPivot(PivotType, PivotActor.Get());
+	}
+	return FReply::Handled();
+}
+
+int32 FLevelInstancePivotDetails::GetSelectedPivotType() const
+{
+	return (int32)PivotType;
+}
+
+void FLevelInstancePivotDetails::OnSelectedPivotTypeChanged(int32 NewValue, ESelectInfo::Type SelectionType)
+{
+	PivotType = (ELevelInstancePivotType)NewValue;
+}
+
+TSharedRef<SWidget> FLevelInstancePivotDetails::OnGeneratePivotActorWidget(TSharedPtr<FActorInfo> ActorInfo) const
+{
+	AActor* Actor = ActorInfo->Actor.Get();
+
+	// If a row wasn't generated just create the default one, a simple text block of the item's name.
+	return SNew(STextBlock).Text(Actor ? FText::FromString(Actor->GetActorLabel()) : LOCTEXT("null", "null"));
+}
+
+FText FLevelInstancePivotDetails::GetSelectedPivotActorText() const
+{
+	return PivotActor.IsValid() ? FText::FromString(PivotActor->GetActorLabel()) : LOCTEXT("none", "None");
+}
+
+void FLevelInstancePivotDetails::OnSelectedPivotActorChanged(TSharedPtr<FActorInfo> NewValue, ESelectInfo::Type SelectionType)
+{
+	PivotActor = NewValue->Actor;
+}
+
+bool FLevelInstancePivotDetails::IsPivotActorSelectionEnabled() const
+{
+	return PivotType == ELevelInstancePivotType::Actor;
+}
+
+
+
+#undef LOCTEXT_NAMESPACE
\ No newline at end of file
diff --git a/Engine/Source/Editor/LevelInstanceEditor/Private/LevelInstancePivotDetails.h b/Engine/Source/Editor/LevelInstanceEditor/Private/LevelInstancePivotDetails.h
new file mode 100644
--- /dev/null
+++ b/Engine/Source/Editor/LevelInstanceEditor/Private/LevelInstancePivotDetails.h
@@ -0,0 +1,47 @@
+// Copyright Epic Games, Inc. All Rights Reserved.
+
+#pragma once
+
+#include "CoreMinimal.h"
+#include "Types/SlateEnums.h"
+#include "IDetailCustomization.h"
+#include "LevelInstance/LevelInstanceTypes.h"
+
+class ALevelInstancePivot;
+
+class FLevelInstancePivotDetails : public IDetailCustomization
+{
+public:
+	/** Makes a new instance of this detail layout class for a specific detail view requesting it */
+	static TSharedRef<IDetailCustomization> MakeInstance();
+
+	/** IDetailCustomization interface */
+	virtual void CustomizeDetails(IDetailLayoutBuilder& DetailBuilder) override;
+	/** End IDetailCustomization interface */
+
+private:
+	/** Use MakeInstance to create an instance of this class */
+	FLevelInstancePivotDetails();
+
+	bool IsApplyButtonEnabled(TWeakObjectPtr<ALevelInstancePivot> LevelInstancePivot) const;
+	FReply OnApplyButtonClicked(TWeakObjectPtr<ALevelInstancePivot> LevelInstancePivot);
+	
+	int32 GetSelectedPivotType() const;
+	void OnSelectedPivotTypeChanged(int32 NewValue, ESelectInfo::Type SelectionType);
+
+	// struct to hold weak ptr
+	struct FActorInfo
+	{
+		TWeakObjectPtr<AActor> Actor;
+	};
+
+	TSharedRef<SWidget> OnGeneratePivotActorWidget(TSharedPtr<FActorInfo> Actor) const;
+	FText GetSelectedPivotActorText() const;
+	void OnSelectedPivotActorChanged(TSharedPtr<FActorInfo> NewValue, ESelectInfo::Type SelectionType);
+	bool IsPivotActorSelectionEnabled() const;
+	
+	// Settings for Apply button
+	ELevelInstancePivotType PivotType;
+	TWeakObjectPtr<AActor> PivotActor;
+	TArray<TSharedPtr<FActorInfo>> PivotActors;
+};
\ No newline at end of file
diff --git a/Engine/Source/Editor/LevelInstanceEditor/Private/SNewLevelInstanceDialog.cpp b/Engine/Source/Editor/LevelInstanceEditor/Private/SNewLevelInstanceDialog.cpp
new file mode 100644
--- /dev/null
+++ b/Engine/Source/Editor/LevelInstanceEditor/Private/SNewLevelInstanceDialog.cpp
@@ -0,0 +1,197 @@
+// Copyright Epic Games, Inc. All Rights Reserved.
+#include "SNewLevelInstanceDialog.h"
+#include "SEnumCombo.h"
+#include "Widgets/Layout/SSpacer.h"
+#include "LevelInstance/LevelInstanceSubsystem.h"
+
+#define LOCTEXT_NAMESPACE "LevelInstanceEditor"
+
+const FVector2D SNewLevelInstanceDialog::DEFAULT_WINDOW_SIZE = FVector2D(400, 150);
+
+void SNewLevelInstanceDialog::Construct(const FArguments& InArgs)
+{
+	SelectedCreationType = ELevelInstanceCreationType::LevelInstance;
+	SelectedPivotType = ELevelInstancePivotType::CenterMinZ;
+	ParentWindowPtr = InArgs._ParentWindow.Get();
+	PivotActors = InArgs._PivotActors.Get();
+	bClickedOk = false;
+
+	this->ChildSlot
+	[
+		SNew(SBorder)
+		.BorderImage(FEditorStyle::GetBrush("ToolPanel.GroupBorder"))
+		[
+			SNew(SVerticalBox)
+			+ SVerticalBox::Slot()
+			.AutoHeight()
+			[
+				SNew(SHorizontalBox)
+				+ SHorizontalBox::Slot()
+				.AutoWidth()
+				.Padding(2.f)
+				.VAlign(VAlign_Center)
+				[
+					SNew(STextBlock)
+					.Text(LOCTEXT("LevelInstanceCreationTypeTextBlock", "Type"))
+				]
+				+ SHorizontalBox::Slot()
+				.FillWidth(1.f)
+				.Padding(2.f)
+				[
+					SNew(SEnumComboBox, StaticEnum<ELevelInstanceCreationType>())
+					.ContentPadding(FCoreStyle::Get().GetMargin("StandardDialog.ContentPadding"))
+					.CurrentValue(this, &SNewLevelInstanceDialog::GetSelectedCreationType)
+					.OnEnumSelectionChanged(this, &SNewLevelInstanceDialog::OnSelectedCreationTypeChanged)
+				]
+			]
+			+ SVerticalBox::Slot()
+			.AutoHeight()
+			[
+				SNew(SHorizontalBox)
+				+ SHorizontalBox::Slot()
+				.AutoWidth()
+				.Padding(2.f)
+				.VAlign(VAlign_Center)
+				[
+					SNew(STextBlock)
+					.Text(LOCTEXT("LevelInstancePivotTypeTextBlock", "Pivot"))
+				]
+				+ SHorizontalBox::Slot()
+				.FillWidth(1.f)
+				.Padding(2.f)
+				[
+					SNew(SEnumComboBox, StaticEnum<ELevelInstancePivotType>())
+					.ContentPadding(FCoreStyle::Get().GetMargin("StandardDialog.ContentPadding"))
+					.CurrentValue(this, &SNewLevelInstanceDialog::GetSelectedPivotType)
+					.OnEnumSelectionChanged(this, &SNewLevelInstanceDialog::OnSelectedPivotTypeChanged)
+				]
+			]
+			+ SVerticalBox::Slot()
+			.AutoHeight()
+			[
+				SNew(SHorizontalBox)
+				.IsEnabled(this, &SNewLevelInstanceDialog::IsPivotActorSelectionEnabled)
+				+ SHorizontalBox::Slot()
+				.AutoWidth()
+				.Padding(2.f)
+				.VAlign(VAlign_Center)
+				[
+					SNew(STextBlock)
+					.Text(LOCTEXT("LevelInstancePivotActorTextBlock", "Actor"))
+				]
+				+ SHorizontalBox::Slot()
+				.FillWidth(1.f)
+				.Padding(2.f)
+				[
+					SNew(SComboBox<AActor*>)
+					.ContentPadding(FCoreStyle::Get().GetMargin("StandardDialog.ContentPadding"))
+					.OptionsSource(&PivotActors)
+					.OnGenerateWidget(this, &SNewLevelInstanceDialog::OnGeneratePivotActorWidget)
+					.OnSelectionChanged(this, &SNewLevelInstanceDialog::OnSelectedPivotActorChanged)
+					[
+						SNew(STextBlock)
+						.Text(this, &SNewLevelInstanceDialog::GetSelectedPivotActorText)
+					]
+				]
+			]
+			+ SVerticalBox::Slot()
+			.AutoHeight()
+			[
+				SNew(SHorizontalBox)
+				+ SHorizontalBox::Slot()
+				.FillWidth(1.f)
+				[
+					SNew(SSpacer)
+				]
+				+ SHorizontalBox::Slot()
+				.AutoWidth()
+				.Padding(2.f)
+				[
+					SNew(SButton)
+					.HAlign(HAlign_Center)
+					.IsEnabled(this, &SNewLevelInstanceDialog::IsOkEnabled)
+					.ContentPadding(FEditorStyle::GetMargin("StandardDialog.ContentPadding"))
+					.OnClicked(this, &SNewLevelInstanceDialog::OnOkClicked)
+					.Text(LOCTEXT("OkButton", "Ok"))
+				]
+				+ SHorizontalBox::Slot()
+				.AutoWidth()
+				.Padding(2.f)
+				[
+					SNew(SButton)
+					.HAlign(HAlign_Center)
+					.ContentPadding(FEditorStyle::GetMargin("StandardDialog.ContentPadding"))
+					.OnClicked(this, &SNewLevelInstanceDialog::OnCancelClicked)
+					.Text(LOCTEXT("CancelButton", "Cancel"))
+				]
+			]
+		]
+	];
+}
+
+bool SNewLevelInstanceDialog::IsOkEnabled() const
+{
+	if (SelectedPivotType == ELevelInstancePivotType::Actor && !SelectedPivotActor)
+	{
+		return false;
+	}
+
+	return true;
+}
+
+FReply SNewLevelInstanceDialog::OnOkClicked()
+{
+	bClickedOk = true;
+	ParentWindowPtr.Pin()->RequestDestroyWindow();
+	return FReply::Handled();
+}
+
+FReply SNewLevelInstanceDialog::OnCancelClicked()
+{
+	bClickedOk = false;
+	ParentWindowPtr.Pin()->RequestDestroyWindow();
+	return FReply::Handled();
+}
+
+int32 SNewLevelInstanceDialog::GetSelectedCreationType() const
+{
+	return (int32)SelectedCreationType;
+}
+
+void SNewLevelInstanceDialog::OnSelectedCreationTypeChanged(int32 NewValue, ESelectInfo::Type SelectionType)
+{
+	SelectedCreationType = (ELevelInstanceCreationType)NewValue;
+}
+
+int32 SNewLevelInstanceDialog::GetSelectedPivotType() const
+{
+	return (int32)SelectedPivotType;
+}
+
+void SNewLevelInstanceDialog::OnSelectedPivotTypeChanged(int32 NewValue, ESelectInfo::Type SelectionType)
+{
+	SelectedPivotType = (ELevelInstancePivotType)NewValue;
+}
+
+TSharedRef<SWidget> SNewLevelInstanceDialog::OnGeneratePivotActorWidget(AActor* Actor) const
+{
+	// If a row wasn't generated just create the default one, a simple text block of the item's name.
+	return SNew(STextBlock).Text(Actor ? FText::FromString(Actor->GetActorLabel()) : LOCTEXT("null", "null"));
+}
+
+FText SNewLevelInstanceDialog::GetSelectedPivotActorText() const
+{
+	return SelectedPivotActor ? FText::FromString(SelectedPivotActor->GetActorLabel()) : LOCTEXT("none", "None");
+}
+
+void SNewLevelInstanceDialog::OnSelectedPivotActorChanged(AActor* NewValue, ESelectInfo::Type SelectionType)
+{
+	SelectedPivotActor = NewValue;
+}
+
+bool SNewLevelInstanceDialog::IsPivotActorSelectionEnabled() const
+{
+	return SelectedPivotType == ELevelInstancePivotType::Actor;
+}
+
+#undef LOCTEXT_NAMESPACE
diff --git a/Engine/Source/Editor/LevelInstanceEditor/Private/SNewLevelInstanceDialog.h b/Engine/Source/Editor/LevelInstanceEditor/Private/SNewLevelInstanceDialog.h
new file mode 100644
--- /dev/null
+++ b/Engine/Source/Editor/LevelInstanceEditor/Private/SNewLevelInstanceDialog.h
@@ -0,0 +1,82 @@
+// Copyright Epic Games, Inc. All Rights Reserved.
+#pragma once
+
+#include "CoreMinimal.h"
+#include "InputCoreTypes.h"
+#include "Types/SlateEnums.h"
+#include "Styling/SlateColor.h"
+#include "Input/Reply.h"
+#include "Widgets/SWidget.h"
+#include "Widgets/SWindow.h"
+#include "Widgets/DeclarativeSyntaxSupport.h"
+#include "Widgets/SCompoundWidget.h"
+
+enum class ELevelInstanceCreationType : uint8;
+
+enum class ELevelInstancePivotType : uint8;
+
+//////////////////////////////////////////////////////////////////////////
+// SNewLevelInstanceDialog
+
+class SNewLevelInstanceDialog : public SCompoundWidget
+{
+public:
+
+	SLATE_BEGIN_ARGS(SNewLevelInstanceDialog) {}
+		/** A pointer to the parent window */
+		SLATE_ATTRIBUTE(TSharedPtr<SWindow>, ParentWindow)
+		SLATE_ATTRIBUTE(TArray<AActor*>, PivotActors)
+	SLATE_END_ARGS()
+
+	~SNewLevelInstanceDialog() {}
+
+	/**
+	 * Construct this widget.  Called by the SNew() Slate macro.
+	 *
+	 * @param	InArgs				Declaration used by the SNew() macro to construct this widget
+	 * @param	InViewModel			The UI logic not specific to slate
+	 */
+	void Construct(const FArguments& InArgs);
+
+	static const FVector2D DEFAULT_WINDOW_SIZE;
+
+	bool ClickedOk() const { return bClickedOk; }
+	ELevelInstanceCreationType GetCreationType() const { return SelectedCreationType; }
+	ELevelInstancePivotType GetPivotType() const { return SelectedPivotType; }
+	AActor* GetPivotActor() const { return SelectedPivotActor; }
+
+private:
+	int32 GetSelectedCreationType() const;
+	void OnSelectedCreationTypeChanged(int32 NewValue, ESelectInfo::Type SelectionType);
+
+	int32 GetSelectedPivotType() const;
+	void OnSelectedPivotTypeChanged(int32 NewValue, ESelectInfo::Type SelectionType);
+		
+	FReply OnOkClicked();
+	bool IsOkEnabled() const;
+
+	FReply OnCancelClicked();
+
+	TSharedRef<SWidget> OnGeneratePivotActorWidget(AActor* Actor) const;
+	FText GetSelectedPivotActorText() const;
+	void OnSelectedPivotActorChanged(AActor* NewValue, ESelectInfo::Type SelectionType);
+	bool IsPivotActorSelectionEnabled() const;
+
+	/** Pointer to the parent window, so we know to destroy it when done */
+	TWeakPtr<SWindow> ParentWindowPtr;
+
+	/** The type of LevelInstance to create */
+	ELevelInstanceCreationType SelectedCreationType;
+
+	/** The type of Pivot */
+	ELevelInstancePivotType SelectedPivotType;
+
+	/** Pivot Actor */
+	AActor* SelectedPivotActor = nullptr;
+
+	/** Actor List */
+	TArray<AActor*> PivotActors;
+
+	/** Dialog Result */
+	bool bClickedOk;
+};
\ No newline at end of file
diff --git a/Engine/Source/Runtime/Engine/Classes/GameFramework/WorldSettings.h b/Engine/Source/Runtime/Engine/Classes/GameFramework/WorldSettings.h
--- a/Engine/Source/Runtime/Engine/Classes/GameFramework/WorldSettings.h
+++ b/Engine/Source/Runtime/Engine/Classes/GameFramework/WorldSettings.h
@@ -484,6 +484,10 @@ public:
 	UPROPERTY(EditAnywhere, Category = WorldPartition, AdvancedDisplay, meta = (editcondition = "bEnableWorldPartition"))
 	TSubclassOf<UWorldPartitionStreamingPolicy> WorldPartitionStreamingPolicyClass;
 
+	/** Additional transform applied when applying LevelStreaming Transform to LevelInstance */
+	UPROPERTY()
+	FVector LevelInstancePivotOffset;
+
 	/**
 	 * Enables client-side streaming volumes instead of server-side.
 	 * Expected usage scenario: server has all streaming levels always loaded, clients independently stream levels in/out based on streaming volumes.
diff --git a/Engine/Source/Runtime/Engine/Private/LevelInstance/LevelInstanceActor.cpp b/Engine/Source/Runtime/Engine/Private/LevelInstance/LevelInstanceActor.cpp
--- a/Engine/Source/Runtime/Engine/Private/LevelInstance/LevelInstanceActor.cpp
+++ b/Engine/Source/Runtime/Engine/Private/LevelInstance/LevelInstanceActor.cpp
@@ -205,6 +205,16 @@ AActor* ALevelInstance::FindEditorInstanceActor() const
 	return FoundActor;
 }
 
+void ALevelInstance::OnEdit()
+{
+	bLockLocation = true;
+}
+
+void ALevelInstance::OnCommit()
+{
+	bLockLocation = false;
+}
+
 ALevelInstance::FOnLevelInstanceActorPostLoad ALevelInstance::OnLevelInstanceActorPostLoad;
 
 void ALevelInstance::PostLoad()
diff --git a/Engine/Source/Runtime/Engine/Private/LevelInstance/LevelInstanceEditorLevelStreaming.cpp b/Engine/Source/Runtime/Engine/Private/LevelInstance/LevelInstanceEditorLevelStreaming.cpp
--- a/Engine/Source/Runtime/Engine/Private/LevelInstance/LevelInstanceEditorLevelStreaming.cpp
+++ b/Engine/Source/Runtime/Engine/Private/LevelInstance/LevelInstanceEditorLevelStreaming.cpp
@@ -5,8 +5,10 @@
 #if WITH_EDITOR
 #include "LevelInstance/LevelInstanceActor.h"
 #include "LevelInstance/LevelInstanceSubsystem.h"
+#include "LevelInstance/LevelInstanceEditorPivotActor.h"
 #include "EditorLevelUtils.h"
 #include "Editor.h"
+#include "Engine/LevelBounds.h"
 #endif
 
 #if WITH_EDITOR
@@ -51,6 +53,9 @@ ULevelStreamingLevelInstanceEditor* ULevelStreamingLevelInstanceEditor::Load(ALe
 
 	GEngine->BlockTillLevelStreamingCompleted(LevelInstanceActor->GetWorld());
 	
+	// Create special actor that will handle changing the pivot of this level
+	ALevelInstancePivot::Create(LevelInstanceActor, LevelStreaming);
+
 	return LevelStreaming;
 }
 
@@ -69,4 +74,24 @@ void ULevelStreamingLevelInstanceEditor::OnLevelActorAdded(AActor* InActor)
 	}
 }
 
+void ULevelStreamingLevelInstanceEditor::SetLoadedLevel(ULevel* Level)
+{
+	Super::SetLoadedLevel(Level);
+
+	if (ULevel* NewLoadedLevel = GetLoadedLevel())
+	{
+		check(!NewLoadedLevel->bAlreadyMovedActors);
+		if (AWorldSettings* WorldSettings = NewLoadedLevel->GetWorldSettings())
+		{
+			LevelTransform.AddToTranslation(WorldSettings->LevelInstancePivotOffset);
+		}
+	}
+}
+
+FBox ULevelStreamingLevelInstanceEditor::GetBounds() const
+{
+	check(GetLoadedLevel());
+	return ALevelBounds::CalculateLevelBounds(GetLoadedLevel());
+}
+
 #endif
\ No newline at end of file
diff --git a/Engine/Source/Runtime/Engine/Private/LevelInstance/LevelInstanceEditorPivotActor.cpp b/Engine/Source/Runtime/Engine/Private/LevelInstance/LevelInstanceEditorPivotActor.cpp
new file mode 100644
--- /dev/null
+++ b/Engine/Source/Runtime/Engine/Private/LevelInstance/LevelInstanceEditorPivotActor.cpp
@@ -0,0 +1,122 @@
+// Copyright Epic Games, Inc. All Rights Reserved.
+
+#include "LevelInstance/LevelInstanceEditorPivotActor.h"
+#include "LevelInstance/LevelInstanceActor.h"
+#include "LevelInstance/LevelInstanceSubsystem.h"
+#include "GameFramework/WorldSettings.h"
+#include "LevelUtils.h"
+#include "Engine/LevelStreaming.h"
+
+#if WITH_EDITOR
+#include "Editor.h"
+#endif
+
+ALevelInstancePivot::ALevelInstancePivot(const FObjectInitializer& ObjectInitializer)
+	: Super(ObjectInitializer)
+{
+	RootComponent = CreateDefaultSubobject<USceneComponent>(TEXT("RootComponent"));
+	RootComponent->Mobility = EComponentMobility::Static;
+}
+
+#if WITH_EDITOR
+
+ALevelInstancePivot* ALevelInstancePivot::Create(ALevelInstance* LevelInstanceActor, ULevelStreaming* LevelStreaming)
+{
+	FActorSpawnParameters SpawnParams;
+	SpawnParams.OverrideLevel = LevelStreaming->GetLoadedLevel();
+	SpawnParams.bCreateActorPackage = false;
+	SpawnParams.ObjectFlags |= RF_Transient;
+	SpawnParams.bNoFail = true;
+		
+	// We place the pivot actor at the LevelInstance Transform so that it makes sense to the user (the pivot being the zero)
+	ALevelInstancePivot* PivotActor = LevelInstanceActor->GetWorld()->SpawnActor<ALevelInstancePivot>(LevelInstanceActor->GetActorLocation(), LevelInstanceActor->GetActorRotation(), SpawnParams);
+	
+	AWorldSettings* WorldSettings = LevelStreaming->GetLoadedLevel()->GetWorldSettings();
+
+	
+	// Keep Spawn World Transform in case Level Instance transform changes.
+	PivotActor->SpawnOffset = PivotActor->GetActorLocation();
+	PivotActor->OriginalPivotOffset = LevelInstanceActor->GetActorLocation() - WorldSettings->LevelInstancePivotOffset;
+	PivotActor->SetLevelInstanceID(LevelInstanceActor->GetLevelInstanceID());
+
+	return PivotActor;
+}
+
+void ALevelInstancePivot::PostEditMove(bool bFinished)
+{
+	Super::PostEditMove(bFinished);
+
+	if (bFinished && !FLevelUtils::IsApplyingLevelTransform())
+	{
+		UpdateOffset();
+	}
+}
+
+void ALevelInstancePivot::PostEditChangeProperty(FPropertyChangedEvent& PropertyChangedEvent)
+{
+	Super::PostEditChangeProperty(PropertyChangedEvent);
+
+	UpdateOffset();
+}
+
+void ALevelInstancePivot::PostEditUndo()
+{
+	Super::PostEditUndo();
+
+	UpdateOffset();
+}
+
+void ALevelInstancePivot::SetPivot(ELevelInstancePivotType PivotType, AActor* PivotActor)
+{
+	check(PivotType != ELevelInstancePivotType::Actor || PivotActor != nullptr);
+
+	Modify();
+	if (PivotType == ELevelInstancePivotType::Actor)
+	{
+		SetActorLocation(PivotActor->GetActorLocation());
+	}
+	else if(PivotType == ELevelInstancePivotType::Center || PivotType == ELevelInstancePivotType::CenterMinZ)
+	{
+		ULevelInstanceSubsystem* LevelInstanceSubsystem = UWorld::GetSubsystem<ULevelInstanceSubsystem>(GetWorld());
+		ALevelInstance* LevelInstance = LevelInstanceSubsystem->GetLevelInstance(GetLevelInstanceID());
+		FBox OutBounds;
+		LevelInstanceSubsystem->GetLevelInstanceBounds(LevelInstance, OutBounds);
+		FVector Location = OutBounds.GetCenter();
+		if (PivotType == ELevelInstancePivotType::CenterMinZ)
+		{
+			Location.Z = OutBounds.Min.Z;
+		}
+		SetActorLocation(Location);
+	}
+	else // unsupported
+	{
+		check(0);
+	}
+
+	// Update gizmo
+	if (GEditor)
+	{
+		GEditor->NoteSelectionChange();
+	}
+
+	UpdateOffset();
+}
+
+void ALevelInstancePivot::UpdateOffset()
+{
+	AWorldSettings* WorldSettings = GetLevel()->GetWorldSettings();
+
+	// Offset Change is the relative transform to the spawn transform
+	FVector LocalToPivot = GetActorLocation() - SpawnOffset;
+
+	// We then apply that offset to the original pivot world transform
+	FVector NewPivotOffset = LocalToPivot + OriginalPivotOffset;
+
+	if (!NewPivotOffset.Equals(WorldSettings->LevelInstancePivotOffset))
+	{
+		WorldSettings->Modify();
+		WorldSettings->LevelInstancePivotOffset = SpawnOffset - NewPivotOffset;
+	}
+}
+
+#endif
diff --git a/Engine/Source/Runtime/Engine/Private/LevelInstance/LevelInstanceLevelStreaming.cpp b/Engine/Source/Runtime/Engine/Private/LevelInstance/LevelInstanceLevelStreaming.cpp
--- a/Engine/Source/Runtime/Engine/Private/LevelInstance/LevelInstanceLevelStreaming.cpp
+++ b/Engine/Source/Runtime/Engine/Private/LevelInstance/LevelInstanceLevelStreaming.cpp
@@ -147,3 +147,16 @@ void ULevelStreamingLevelInstance::UnloadInstance(ULevelStreamingLevelInstance*
 #endif 
 }
 
+void ULevelStreamingLevelInstance::SetLoadedLevel(ULevel* Level)
+{
+	Super::SetLoadedLevel(Level);
+
+	if (ULevel* NewLoadedLevel = GetLoadedLevel())
+	{
+		check(!NewLoadedLevel->bAlreadyMovedActors);
+		if (AWorldSettings* WorldSettings = NewLoadedLevel->GetWorldSettings())
+		{
+			LevelTransform.AddToTranslation(WorldSettings->LevelInstancePivotOffset);
+		}
+	}
+}
diff --git a/Engine/Source/Runtime/Engine/Private/LevelInstance/LevelInstanceSubsystem.cpp b/Engine/Source/Runtime/Engine/Private/LevelInstance/LevelInstanceSubsystem.cpp
--- a/Engine/Source/Runtime/Engine/Private/LevelInstance/LevelInstanceSubsystem.cpp
+++ b/Engine/Source/Runtime/Engine/Private/LevelInstance/LevelInstanceSubsystem.cpp
@@ -446,6 +446,11 @@ bool ULevelInstanceSubsystem::GetLevelInstanceBounds(const ALevelInstance* Level
 		OutBounds = LevelInstance.LevelStreaming->GetBounds();
 		return true;
 	}
+	else if (const FLevelInstanceEdit* LevelInstanceEdit = GetLevelInstanceEdit(LevelInstanceActor))
+	{
+		OutBounds = LevelInstanceEdit->LevelStreaming->GetBounds();
+		return true;
+	}
 	else if(LevelInstanceActor->IsLevelInstancePathValid())
 	{
 		if (ULevel::GetLevelBoundsFromPackage(FName(*LevelInstanceActor->GetWorldAssetPackage()), OutBounds))
@@ -633,7 +638,7 @@ bool ULevelInstanceSubsystem::MoveActorsTo(ALevelInstance* LevelInstanceActor, c
 	return MoveActorsToLevel(ActorsToMove, LevelInstanceLevel);
 }
 
-ALevelInstance* ULevelInstanceSubsystem::CreateLevelInstanceFrom(const TArray<AActor*>& ActorsToMove, ELevelInstanceCreationType CreationType, UWorld* TemplateWorld)
+ALevelInstance* ULevelInstanceSubsystem::CreateLevelInstanceFrom(const TArray<AActor*>& ActorsToMove, ELevelInstanceCreationType CreationType, ELevelInstancePivotType PivotType, AActor* PivotActor, UWorld* TemplateWorld)
 {
 	ULevel* CurrentLevel = GetWorld()->GetCurrentLevel();
 		
@@ -658,9 +663,21 @@ ALevelInstance* ULevelInstanceSubsystem::CreateLevelInstanceFrom(const TArray<AA
 		}
 	}
 
-	FVector LevelInstanceLocation = ActorLocationBox.GetCenter();
-	LevelInstanceLocation.Z = ActorLocationBox.Min.Z;
-
+	FVector LevelInstanceLocation;
+	if (PivotType == ELevelInstancePivotType::Actor)
+	{
+		check(PivotActor);
+		LevelInstanceLocation = PivotActor->GetActorLocation();
+	}
+	else
+	{
+		LevelInstanceLocation = ActorLocationBox.GetCenter();
+		if (PivotType == ELevelInstancePivotType::CenterMinZ)
+		{
+			LevelInstanceLocation.Z = ActorLocationBox.Min.Z;
+		}
+	}
+		
 	ULevelStreamingLevelInstanceEditor* LevelStreaming = StaticCast<ULevelStreamingLevelInstanceEditor*>(EditorLevelUtils::CreateNewStreamingLevelForWorld(*GetWorld(), ULevelStreamingLevelInstanceEditor::StaticClass(), TEXT(""), false, TemplateWorld));
 	if (!LevelStreaming)
 	{
@@ -730,7 +747,7 @@ ALevelInstance* ULevelInstanceSubsystem::CreateLevelInstanceFrom(const TArray<AA
 	check(NewLevelInstanceActor);
 	NewLevelInstanceActor->SetWorldAsset(WorldPtr);
 	NewLevelInstanceActor->SetActorLocation(LevelInstanceLocation);
-
+	
 	// Actors were moved and kept their World positions so when saving we want their positions to actually be relative to the FounationActor/LevelTransform
 	// so we set the LevelTransform and we mark the level as having moved its actors. 
 	// On Level save FLevelUtils::RemoveEditorTransform will fixup actor transforms to make them relative to the LevelTransform.
diff --git a/Engine/Source/Runtime/Engine/Private/LevelInstance/Packed/LevelInstanceISMPacker.cpp b/Engine/Source/Runtime/Engine/Private/LevelInstance/Packed/LevelInstanceISMPacker.cpp
--- a/Engine/Source/Runtime/Engine/Private/LevelInstance/Packed/LevelInstanceISMPacker.cpp
+++ b/Engine/Source/Runtime/Engine/Private/LevelInstance/Packed/LevelInstanceISMPacker.cpp
@@ -39,7 +39,10 @@ void FLevelInstanceISMPacker::PackActors(FPackedLevelInstanceBuilderContext& InC
 	check(InClusterID.GetPackerID() == GetID());
 	UInstancedStaticMeshComponent* PackComponent = InPackingActor->AddPackedComponent<UInstancedStaticMeshComponent>();
 	
-	PackComponent->SetComponentToWorld(InPackingActor->GetActorTransform());
+	FTransform ActorTransform = InPackingActor->GetActorTransform();
+	FTransform CurrentPivotOffsetInverse = ActorTransform.GetRelativeTransform(InContext.GetLevelTransform());
+
+	PackComponent->SetComponentToWorld(ActorTransform);
 	PackComponent->AttachToComponent(InPackingActor->GetRootComponent(), FAttachmentTransformRules::KeepWorldTransform);
 	
 	FLevelInstanceISMPackerCluster* ISMCluster = (FLevelInstanceISMPackerCluster*)InClusterID.GetData();
@@ -63,7 +66,7 @@ void FLevelInstanceISMPacker::PackActors(FPackedLevelInstanceBuilderContext& InC
 			{
 				FTransform InstanceTransform;
 
-				if (ensure(ISMComponent->GetInstanceTransform(InstanceIndex, InstanceTransform, true)))
+				if (ensure(ISMComponent->GetInstanceTransform(InstanceIndex, InstanceTransform, /*bWorldSpace=*/ true)))
 				{
 					PackComponent->AddInstanceWorldSpace(InstanceTransform);
 				}
@@ -75,6 +78,10 @@ void FLevelInstanceISMPacker::PackActors(FPackedLevelInstanceBuilderContext& InC
 			PackComponent->AddInstanceWorldSpace(StaticMeshComponent->GetComponentTransform());
 		}
 	}
+
+	FTransform NewWorldTransform = ActorTransform * CurrentPivotOffsetInverse * FTransform(InContext.GetPivotOffset());
+
+	PackComponent->SetWorldTransform(NewWorldTransform);
 	PackComponent->RegisterComponent();
 }
 
diff --git a/Engine/Source/Runtime/Engine/Private/LevelInstance/Packed/PackedLevelInstanceActor.cpp b/Engine/Source/Runtime/Engine/Private/LevelInstance/Packed/PackedLevelInstanceActor.cpp
--- a/Engine/Source/Runtime/Engine/Private/LevelInstance/Packed/PackedLevelInstanceActor.cpp
+++ b/Engine/Source/Runtime/Engine/Private/LevelInstance/Packed/PackedLevelInstanceActor.cpp
@@ -93,6 +93,7 @@ void APackedLevelInstance::OnCommitChild()
 
 void APackedLevelInstance::OnEdit()
 {
+	Super::OnEdit();
 	check(!bEditing);
 	
 	bEditing = true;
@@ -101,6 +102,7 @@ void APackedLevelInstance::OnEdit()
 
 void APackedLevelInstance::OnCommit()
 {
+	Super::OnCommit();
 	check(bEditing);
 	bEditing = false;
 	MarkComponentsRenderStateDirty();
diff --git a/Engine/Source/Runtime/Engine/Private/LevelInstance/Packed/PackedLevelInstanceBuilder.cpp b/Engine/Source/Runtime/Engine/Private/LevelInstance/Packed/PackedLevelInstanceBuilder.cpp
--- a/Engine/Source/Runtime/Engine/Private/LevelInstance/Packed/PackedLevelInstanceBuilder.cpp
+++ b/Engine/Source/Runtime/Engine/Private/LevelInstance/Packed/PackedLevelInstanceBuilder.cpp
@@ -23,6 +23,7 @@
 #include "Components/SceneComponent.h"
 #include "GameFramework/WorldSettings.h"
 #include "Engine/LevelBounds.h"
+#include "Engine/LevelStreaming.h"
 #include "Engine/Brush.h"
 #include "FileHelpers.h"
 #include "Editor.h"
@@ -135,7 +136,12 @@ void FPackedLevelInstanceBuilder::PackActor(APackedLevelInstance* InPackedLevelI
 	check(LevelInstanceSubystem);
 	
 	ULevel* SourceLevel = LevelInstanceSubystem->GetLevelInstanceLevel(InLevelInstanceToPack);
-		
+	ULevelStreaming* SourceLevelStreaming = ULevelStreaming::FindStreamingLevel(SourceLevel);
+	AWorldSettings* WorldSettings = SourceLevel->GetWorldSettings();
+	
+	Context.SetLevelTransform(SourceLevelStreaming->LevelTransform);
+	Context.SetPivotOffset(WorldSettings->LevelInstancePivotOffset);
+	
 	if (AActor* DefaultBrush = SourceLevel->GetDefaultBrush())
 	{
 		Context.DiscardActor(DefaultBrush);
diff --git a/Engine/Source/Runtime/Engine/Private/WorldSettings.cpp b/Engine/Source/Runtime/Engine/Private/WorldSettings.cpp
--- a/Engine/Source/Runtime/Engine/Private/WorldSettings.cpp
+++ b/Engine/Source/Runtime/Engine/Private/WorldSettings.cpp
@@ -126,6 +126,8 @@ PRAGMA_ENABLE_DEPRECATION_WARNINGS
 
 	// Default to UWorldPartitionLevelStreamingPolicy
 	WorldPartitionStreamingPolicyClass = UWorldPartitionLevelStreamingPolicy::StaticClass();
+
+	LevelInstancePivotOffset = FVector::ZeroVector;
 }
 
 void AWorldSettings::PostInitProperties()
diff --git a/Engine/Source/Runtime/Engine/Public/LevelInstance/LevelInstanceActor.h b/Engine/Source/Runtime/Engine/Public/LevelInstance/LevelInstanceActor.h
--- a/Engine/Source/Runtime/Engine/Public/LevelInstance/LevelInstanceActor.h
+++ b/Engine/Source/Runtime/Engine/Public/LevelInstance/LevelInstanceActor.h
@@ -79,9 +79,9 @@ public:
 
 	virtual void OnWorldAssetChanged() { UpdateLevelInstance(); }
 	virtual void OnWorldAssetSaved() {}
-	virtual void OnEdit() {}
+	virtual void OnEdit();
 	virtual void OnEditChild() {}
-	virtual void OnCommit() {}
+	virtual void OnCommit();
 	virtual void OnCommitChild() {}
 		
 	DECLARE_MULTICAST_DELEGATE_OneParam(FOnLevelInstanceActorPostLoad, ALevelInstance*);
diff --git a/Engine/Source/Runtime/Engine/Public/LevelInstance/LevelInstanceEditorLevelStreaming.h b/Engine/Source/Runtime/Engine/Public/LevelInstance/LevelInstanceEditorLevelStreaming.h
--- a/Engine/Source/Runtime/Engine/Public/LevelInstance/LevelInstanceEditorLevelStreaming.h
+++ b/Engine/Source/Runtime/Engine/Public/LevelInstance/LevelInstanceEditorLevelStreaming.h
@@ -19,6 +19,7 @@ public:
 #if WITH_EDITOR
 	virtual bool ShowInLevelCollection() const override { return false; }
 	ALevelInstance* GetLevelInstanceActor() const;
+	FBox GetBounds() const;
 
 protected:
 	void OnLevelActorAdded(AActor* InActor);
@@ -27,6 +28,8 @@ protected:
 
 	static ULevelStreamingLevelInstanceEditor* Load(ALevelInstance* LevelInstanceActor);
 	static void Unload(ULevelStreamingLevelInstanceEditor* LevelStreaming);
+
+	virtual void SetLoadedLevel(ULevel* Level) override;
 private:
 	FLevelInstanceID LevelInstanceID;
 #endif
diff --git a/Engine/Source/Runtime/Engine/Public/LevelInstance/LevelInstanceEditorPivotActor.h b/Engine/Source/Runtime/Engine/Public/LevelInstance/LevelInstanceEditorPivotActor.h
new file mode 100644
--- /dev/null
+++ b/Engine/Source/Runtime/Engine/Public/LevelInstance/LevelInstanceEditorPivotActor.h
@@ -0,0 +1,44 @@
+// Copyright Epic Games, Inc. All Rights Reserved.
+
+#pragma once
+
+#include "CoreMinimal.h"
+#include "UObject/ObjectMacros.h"
+#include "GameFramework/Actor.h"
+#include "LevelInstance/LevelInstanceTypes.h"
+#include "LevelInstanceEditorPivotActor.generated.h"
+
+class ALevelInstance;
+class ULevelStreaming;
+
+/**
+ * 
+ */
+UCLASS(transient, notplaceable, hidecategories=(Rendering,Replication,Collision,Partition,Input,HLOD,Actor,LOD,Cooking))
+class ENGINE_API ALevelInstancePivot : public AActor
+{
+	GENERATED_UCLASS_BODY()
+
+#if WITH_EDITOR
+public:
+	static ALevelInstancePivot* Create(ALevelInstance* LevelInstanceActor, ULevelStreaming* LevelStreaming);
+	
+	virtual bool CanDeleteSelectedActor(FText& OutReason) const override { return false; }
+	virtual void PostEditMove(bool bFinished) override;
+	virtual void PostEditChangeProperty(FPropertyChangedEvent& PropertyChangedEvent) override;
+	virtual void PostEditUndo() override;
+
+	void SetPivot(ELevelInstancePivotType PivotType, AActor* PivotActor = nullptr);
+private:
+	void SetLevelInstanceID(const FLevelInstanceID& InLevelInstanceID) { LevelInstanceID = InLevelInstanceID; }
+	const FLevelInstanceID& GetLevelInstanceID() const { return LevelInstanceID; }
+	void UpdateOffset();
+#endif
+
+private:
+#if WITH_EDITORONLY_DATA
+	FLevelInstanceID LevelInstanceID;
+	FVector SpawnOffset;
+	FVector OriginalPivotOffset;
+#endif
+};
\ No newline at end of file
diff --git a/Engine/Source/Runtime/Engine/Public/LevelInstance/LevelInstanceLevelStreaming.h b/Engine/Source/Runtime/Engine/Public/LevelInstance/LevelInstanceLevelStreaming.h
--- a/Engine/Source/Runtime/Engine/Public/LevelInstance/LevelInstanceLevelStreaming.h
+++ b/Engine/Source/Runtime/Engine/Public/LevelInstance/LevelInstanceLevelStreaming.h
@@ -27,6 +27,8 @@ protected:
 	static ULevelStreamingLevelInstance* LoadInstance(ALevelInstance* LevelInstanceActor);
 	static void UnloadInstance(ULevelStreamingLevelInstance* LevelStreaming);
 
+	virtual void SetLoadedLevel(ULevel* Level) override;
+
 	friend class ULevelInstanceSubsystem;
 
 private:
diff --git a/Engine/Source/Runtime/Engine/Public/LevelInstance/LevelInstanceSubsystem.h b/Engine/Source/Runtime/Engine/Public/LevelInstance/LevelInstanceSubsystem.h
--- a/Engine/Source/Runtime/Engine/Public/LevelInstance/LevelInstanceSubsystem.h
+++ b/Engine/Source/Runtime/Engine/Public/LevelInstance/LevelInstanceSubsystem.h
@@ -15,14 +15,6 @@ class ULevelStreamingLevelInstanceEditor;
 class UWorldPartitionSubsystem;
 class UBlueprint;
 
-UENUM()
-enum class ELevelInstanceCreationType : uint8
-{
-	LevelInstance,
-	PackedLevelInstance,
-	PackedLevelInstanceBlueprint
-};
-
 /**
  * ULevelInstanceSubsystem
  */
@@ -81,7 +73,7 @@ public:
 
 	bool SetCurrent(ALevelInstance* LevelInstanceActor) const;
 	bool IsCurrent(const ALevelInstance* LevelInstanceActor) const;
-	ALevelInstance* CreateLevelInstanceFrom(const TArray<AActor*>& ActorsToMove, ELevelInstanceCreationType CreationType = ELevelInstanceCreationType::LevelInstance, UWorld* TemplateWorld = nullptr);
+	ALevelInstance* CreateLevelInstanceFrom(const TArray<AActor*>& ActorsToMove, ELevelInstanceCreationType CreationType, ELevelInstancePivotType PivotType, AActor* PivotActor = nullptr, UWorld* TemplateWorld = nullptr);
 	bool MoveActorsToLevel(const TArray<AActor*>& ActorsToRemove, ULevel* DestinationLevel) const;
 	bool MoveActorsTo(ALevelInstance* LevelInstanceActor, const TArray<AActor*>& ActorsToMove);
 	bool BreakLevelInstance(ALevelInstance* LevelInstanceActor, uint32 Levels = 1);
diff --git a/Engine/Source/Runtime/Engine/Public/LevelInstance/LevelInstanceTypes.h b/Engine/Source/Runtime/Engine/Public/LevelInstance/LevelInstanceTypes.h
--- a/Engine/Source/Runtime/Engine/Public/LevelInstance/LevelInstanceTypes.h
+++ b/Engine/Source/Runtime/Engine/Public/LevelInstance/LevelInstanceTypes.h
@@ -6,4 +6,20 @@
 // Resulting in a different ID for all instances whether they load the same level or not.
 using FLevelInstanceID = uint32;
 
-static constexpr FLevelInstanceID InvalidLevelInstanceID = 0;
\ No newline at end of file
+static constexpr FLevelInstanceID InvalidLevelInstanceID = 0;
+
+UENUM()
+enum class ELevelInstanceCreationType : uint8
+{
+	LevelInstance,
+	PackedLevelInstance,
+	PackedLevelInstanceBlueprint
+};
+
+UENUM()
+enum class ELevelInstancePivotType : uint8
+{
+	CenterMinZ,
+	Center,
+	Actor
+};
\ No newline at end of file
diff --git a/Engine/Source/Runtime/Engine/Public/LevelInstance/Packed/PackedLevelInstanceBuilder.h b/Engine/Source/Runtime/Engine/Public/LevelInstance/Packed/PackedLevelInstanceBuilder.h
--- a/Engine/Source/Runtime/Engine/Public/LevelInstance/Packed/PackedLevelInstanceBuilder.h
+++ b/Engine/Source/Runtime/Engine/Public/LevelInstance/Packed/PackedLevelInstanceBuilder.h
@@ -70,7 +70,8 @@ private:
 class FPackedLevelInstanceBuilderContext
 {
 public:
-	FPackedLevelInstanceBuilderContext(const FPackedLevelInstanceBuilder& InBuilder, APackedLevelInstance* InPackedLevelInstance) : Packers(InBuilder.Packers), ClassDiscards(InBuilder.ClassDiscards), PackedLevelInstance(InPackedLevelInstance) {}
+	FPackedLevelInstanceBuilderContext(const FPackedLevelInstanceBuilder& InBuilder, APackedLevelInstance* InPackedLevelInstance) 
+		: Packers(InBuilder.Packers), ClassDiscards(InBuilder.ClassDiscards), PackedLevelInstance(InPackedLevelInstance), LevelTransform(FVector::ZeroVector), PivotOffset(FVector::ZeroVector) {}
 
 	/* Interface for ILevelInstancePacker's to use */
 	void ClusterLevelActor(AActor* InLevelActor);
@@ -80,6 +81,11 @@ public:
 
 	const TMap<FLevelInstancePackerClusterID, TArray<UActorComponent*>>& GetClusters() const { return Clusters; }
 
+	void SetLevelTransform(const FTransform& InLevelTransform) { LevelTransform = InLevelTransform; }
+	void SetPivotOffset(const FVector& InPivotOffset) { PivotOffset = InPivotOffset; }
+	const FTransform& GetLevelTransform() const { return LevelTransform; }
+	const FVector GetPivotOffset() const { return PivotOffset; }
+
 private:
 	const TMap<FLevelInstancePackerID, TUniquePtr<ILevelInstancePacker>>& Packers;
 	const TSet<UClass*>& ClassDiscards;
@@ -90,6 +96,9 @@ private:
 
 	TMap<AActor*, TSet<UActorComponent*>> PerActorClusteredComponents;
 	TSet<AActor*> ActorDiscards;
+
+	FTransform LevelTransform;
+	FVector PivotOffset;
 };
 
 #endif
\ No newline at end of file
