LevelInstance: Move all creation params into struct and use details panel in New LevelInstance dialog

diff --git a/Engine/Source/Editor/LevelInstanceEditor/Private/LevelInstanceEditorModule.cpp b/Engine/Source/Editor/LevelInstanceEditor/Private/LevelInstanceEditorModule.cpp
--- a/Engine/Source/Editor/LevelInstanceEditor/Private/LevelInstanceEditorModule.cpp
+++ b/Engine/Source/Editor/LevelInstanceEditor/Private/LevelInstanceEditorModule.cpp
@@ -336,18 +336,17 @@ namespace LevelInstanceMenuUtils
 
 		if (NewLevelInstanceDialog->ClickedOk())
 		{
-			ELevelInstanceCreationType CreationType = NewLevelInstanceDialog->GetCreationType();
-			ELevelInstancePivotType PivotType = NewLevelInstanceDialog->GetPivotType();
-			AActor* PivotActor = NewLevelInstanceDialog->GetPivotActor();
+			FNewLevelInstanceParams CreationParams(NewLevelInstanceDialog->GetCreationParams());
 
 			FNewLevelDialogModule& NewLevelDialogModule = FModuleManager::LoadModuleChecked<FNewLevelDialogModule>("NewLevelDialog");
 			FString TemplateMapPackage;
 			if (!GetMutableDefault<ULevelInstanceEditorSettings>()->TemplateMapInfos.Num() || NewLevelDialogModule.CreateAndShowTemplateDialog(MainFrameModule.GetParentWindow(), LOCTEXT("LevelInstanceTemplateDialog", "Choose Level Instance Template..."), GetMutableDefault<ULevelInstanceEditorSettings>()->TemplateMapInfos, TemplateMapPackage))
 			{
 				UPackage* TemplatePackage = !TemplateMapPackage.IsEmpty() ? LoadPackage(nullptr, *TemplateMapPackage, LOAD_None) : nullptr;
-				UWorld* TemplateWorld = TemplatePackage ? UWorld::FindWorldInPackage(TemplatePackage) : nullptr;
+				
+				CreationParams.TemplateWorld = TemplatePackage ? UWorld::FindWorldInPackage(TemplatePackage) : nullptr;
 
-				if (!LevelInstanceSubsystem->CreateLevelInstanceFrom(ActorsToMove, CreationType, PivotType, PivotActor, TemplateWorld))
+				if (!LevelInstanceSubsystem->CreateLevelInstanceFrom(ActorsToMove, CreationParams))
 				{
 					FText Title = LOCTEXT("CreateFromSelectionFailTitle", "Create from selection failed");
 					FMessageDialog::Open(EAppMsgType::Ok, LOCTEXT("CreateFromSelectionFailMsg", "Failed to create LevelInstance from selection. Check log for details."), &Title);
diff --git a/Engine/Source/Editor/LevelInstanceEditor/Private/SNewLevelInstanceDialog.cpp b/Engine/Source/Editor/LevelInstanceEditor/Private/SNewLevelInstanceDialog.cpp
--- a/Engine/Source/Editor/LevelInstanceEditor/Private/SNewLevelInstanceDialog.cpp
+++ b/Engine/Source/Editor/LevelInstanceEditor/Private/SNewLevelInstanceDialog.cpp
@@ -1,20 +1,56 @@
 // Copyright Epic Games, Inc. All Rights Reserved.
 #include "SNewLevelInstanceDialog.h"
-#include "SEnumCombo.h"
+#include "IStructureDetailsView.h"
+#include "IDetailPropertyRow.h"
+#include "DetailCategoryBuilder.h"
+#include "DetailWidgetRow.h"
 #include "Widgets/Layout/SSpacer.h"
+#include "Widgets/Input/SComboBox.h"
+#include "Widgets/Input/SButton.h"
+#include "PropertyEditorDelegates.h"
 #include "LevelInstance/LevelInstanceSubsystem.h"
+#include "Modules/ModuleManager.h"
+#include "GameFramework/Actor.h"
 
 #define LOCTEXT_NAMESPACE "LevelInstanceEditor"
 
-const FVector2D SNewLevelInstanceDialog::DEFAULT_WINDOW_SIZE = FVector2D(400, 150);
+const FVector2D SNewLevelInstanceDialog::DEFAULT_WINDOW_SIZE = FVector2D(400, 250);
 
 void SNewLevelInstanceDialog::Construct(const FArguments& InArgs)
 {
-	SelectedCreationType = ELevelInstanceCreationType::LevelInstance;
-	SelectedPivotType = ELevelInstancePivotType::CenterMinZ;
 	ParentWindowPtr = InArgs._ParentWindow.Get();
-	PivotActors = InArgs._PivotActors.Get();
 	bClickedOk = false;
+		
+	FPropertyEditorModule& PropertyEditorModule = FModuleManager::GetModuleChecked<FPropertyEditorModule>("PropertyEditor");
+
+	TSharedPtr<IStructureDetailsView> StructureDetailsView;
+
+	FDetailsViewArgs DetailsViewArgs;
+	{
+		DetailsViewArgs.bAllowSearch = false;
+		DetailsViewArgs.bHideSelectionTip = true;
+		DetailsViewArgs.bLockable = false;
+		DetailsViewArgs.bSearchInitialKeyFocus = true;
+		DetailsViewArgs.bUpdatesFromSelection = false;
+		DetailsViewArgs.NotifyHook = nullptr;
+		DetailsViewArgs.bShowOptions = true;
+		DetailsViewArgs.bShowModifiedPropertiesOption = false;
+		DetailsViewArgs.bShowScrollBar = false;
+	}
+
+	FStructureDetailsViewArgs StructureViewArgs;
+	{
+		StructureViewArgs.bShowObjects = true;
+		StructureViewArgs.bShowAssets = true;
+		StructureViewArgs.bShowClasses = true;
+		StructureViewArgs.bShowInterfaces = true;
+	}
+
+	StructureDetailsView = PropertyEditorModule.CreateStructureDetailView(DetailsViewArgs, StructureViewArgs, nullptr);
+	StructureDetailsView->GetDetailsView()->SetGenericLayoutDetailsDelegate(FOnGetDetailCustomizationInstance::CreateStatic(&FNewLevelInstanceParamsDetails::MakeInstance, InArgs._PivotActors.Get()));
+
+	FStructOnScope* Struct = new FStructOnScope(FNewLevelInstanceParams::StaticStruct(), (uint8*)&CreationParams);
+	StructureDetailsView->SetStructureData(MakeShareable(Struct));
 
 	this->ChildSlot
 	[
@@ -23,76 +59,9 @@ void SNewLevelInstanceDialog::Construct(const FArguments& InArgs)
 		[
 			SNew(SVerticalBox)
 			+ SVerticalBox::Slot()
-			.AutoHeight()
-			[
-				SNew(SHorizontalBox)
-				+ SHorizontalBox::Slot()
-				.AutoWidth()
-				.Padding(2.f)
-				.VAlign(VAlign_Center)
-				[
-					SNew(STextBlock)
-					.Text(LOCTEXT("LevelInstanceCreationTypeTextBlock", "Type"))
-				]
-				+ SHorizontalBox::Slot()
-				.FillWidth(1.f)
-				.Padding(2.f)
-				[
-					SNew(SEnumComboBox, StaticEnum<ELevelInstanceCreationType>())
-					.ContentPadding(FCoreStyle::Get().GetMargin("StandardDialog.ContentPadding"))
-					.CurrentValue(this, &SNewLevelInstanceDialog::GetSelectedCreationType)
-					.OnEnumSelectionChanged(this, &SNewLevelInstanceDialog::OnSelectedCreationTypeChanged)
-				]
-			]
-			+ SVerticalBox::Slot()
-			.AutoHeight()
-			[
-				SNew(SHorizontalBox)
-				+ SHorizontalBox::Slot()
-				.AutoWidth()
-				.Padding(2.f)
-				.VAlign(VAlign_Center)
-				[
-					SNew(STextBlock)
-					.Text(LOCTEXT("LevelInstancePivotTypeTextBlock", "Pivot"))
-				]
-				+ SHorizontalBox::Slot()
-				.FillWidth(1.f)
-				.Padding(2.f)
-				[
-					SNew(SEnumComboBox, StaticEnum<ELevelInstancePivotType>())
-					.ContentPadding(FCoreStyle::Get().GetMargin("StandardDialog.ContentPadding"))
-					.CurrentValue(this, &SNewLevelInstanceDialog::GetSelectedPivotType)
-					.OnEnumSelectionChanged(this, &SNewLevelInstanceDialog::OnSelectedPivotTypeChanged)
-				]
-			]
-			+ SVerticalBox::Slot()
-			.AutoHeight()
+			.MaxHeight(250.0f)
 			[
-				SNew(SHorizontalBox)
-				.IsEnabled(this, &SNewLevelInstanceDialog::IsPivotActorSelectionEnabled)
-				+ SHorizontalBox::Slot()
-				.AutoWidth()
-				.Padding(2.f)
-				.VAlign(VAlign_Center)
-				[
-					SNew(STextBlock)
-					.Text(LOCTEXT("LevelInstancePivotActorTextBlock", "Actor"))
-				]
-				+ SHorizontalBox::Slot()
-				.FillWidth(1.f)
-				.Padding(2.f)
-				[
-					SNew(SComboBox<AActor*>)
-					.ContentPadding(FCoreStyle::Get().GetMargin("StandardDialog.ContentPadding"))
-					.OptionsSource(&PivotActors)
-					.OnGenerateWidget(this, &SNewLevelInstanceDialog::OnGeneratePivotActorWidget)
-					.OnSelectionChanged(this, &SNewLevelInstanceDialog::OnSelectedPivotActorChanged)
-					[
-						SNew(STextBlock)
-						.Text(this, &SNewLevelInstanceDialog::GetSelectedPivotActorText)
-					]
-				]
+				StructureDetailsView->GetWidget()->AsShared()
 			]
 			+ SVerticalBox::Slot()
 			.AutoHeight()
@@ -131,7 +100,7 @@ void SNewLevelInstanceDialog::Construct(const FArguments& InArgs)
 
 bool SNewLevelInstanceDialog::IsOkEnabled() const
 {
-	if (SelectedPivotType == ELevelInstancePivotType::Actor && !SelectedPivotActor)
+	if (CreationParams.PivotType == ELevelInstancePivotType::Actor && !CreationParams.PivotActor)
 	{
 		return false;
 	}
@@ -153,45 +122,85 @@ FReply SNewLevelInstanceDialog::OnCancelClicked()
 	return FReply::Handled();
 }
 
-int32 SNewLevelInstanceDialog::GetSelectedCreationType() const
+TSharedRef<SWidget> FNewLevelInstanceParamsDetails::OnGeneratePivotActorWidget(AActor* Actor) const
 {
-	return (int32)SelectedCreationType;
+	// If a row wasn't generated just create the default one, a simple text block of the item's name.
+	return SNew(STextBlock).Text(Actor ? FText::FromString(Actor->GetActorLabel()) : LOCTEXT("null", "null")).Font(IDetailLayoutBuilder::GetDetailFont());
 }
 
-void SNewLevelInstanceDialog::OnSelectedCreationTypeChanged(int32 NewValue, ESelectInfo::Type SelectionType)
+FText FNewLevelInstanceParamsDetails::GetSelectedPivotActorText() const
 {
-	SelectedCreationType = (ELevelInstanceCreationType)NewValue;
+	return CreationParams->PivotActor ? FText::FromString(CreationParams->PivotActor->GetActorLabel()) : LOCTEXT("none", "None");
 }
 
-int32 SNewLevelInstanceDialog::GetSelectedPivotType() const
+void FNewLevelInstanceParamsDetails::OnSelectedPivotActorChanged(AActor* NewValue, ESelectInfo::Type SelectionType)
 {
-	return (int32)SelectedPivotType;
+	CreationParams->PivotActor = NewValue;
 }
 
-void SNewLevelInstanceDialog::OnSelectedPivotTypeChanged(int32 NewValue, ESelectInfo::Type SelectionType)
+bool FNewLevelInstanceParamsDetails::IsPivotActorSelectionEnabled() const
 {
-	SelectedPivotType = (ELevelInstancePivotType)NewValue;
+	return CreationParams->PivotType == ELevelInstancePivotType::Actor;
 }
 
-TSharedRef<SWidget> SNewLevelInstanceDialog::OnGeneratePivotActorWidget(AActor* Actor) const
+void FNewLevelInstanceParamsDetails::CustomizeDetails(IDetailLayoutBuilder& DetailBuilder)
 {
-	// If a row wasn't generated just create the default one, a simple text block of the item's name.
-	return SNew(STextBlock).Text(Actor ? FText::FromString(Actor->GetActorLabel()) : LOCTEXT("null", "null"));
-}
+	// Make sure we actually get a valid struct before continuing
+	TArray<TSharedPtr<FStructOnScope>> Structs;
+	DetailBuilder.GetStructsBeingCustomized(Structs);
 
-FText SNewLevelInstanceDialog::GetSelectedPivotActorText() const
-{
-	return SelectedPivotActor ? FText::FromString(SelectedPivotActor->GetActorLabel()) : LOCTEXT("none", "None");
-}
+	if (Structs.Num() == 0)
+	{
+		// Nothing being customized
+		return;
+	}
 
-void SNewLevelInstanceDialog::OnSelectedPivotActorChanged(AActor* NewValue, ESelectInfo::Type SelectionType)
-{
-	SelectedPivotActor = NewValue;
-}
+	const UStruct* Struct = Structs[0]->GetStruct();
+	if (!Struct || Struct != FNewLevelInstanceParams::StaticStruct())
+	{
+		// Invalid struct
+		return;
+	}
 
-bool SNewLevelInstanceDialog::IsPivotActorSelectionEnabled() const
-{
-	return SelectedPivotType == ELevelInstancePivotType::Actor;
+	// Get ptr to our actual type
+	CreationParams = (FNewLevelInstanceParams*)Structs[0]->GetStructMemory();
+
+	TSharedPtr<IPropertyHandle> PivotTypeProperty = DetailBuilder.GetProperty(GET_MEMBER_NAME_CHECKED(FNewLevelInstanceParams, PivotType), (UClass*)FNewLevelInstanceParams::StaticStruct());
+	TSharedPtr<IPropertyHandle> PivotActorProperty = DetailBuilder.GetProperty(GET_MEMBER_NAME_CHECKED(FNewLevelInstanceParams, PivotActor), (UClass*)FNewLevelInstanceParams::StaticStruct());
+
+
+	IDetailCategoryBuilder& PivotCategoryBuilder = DetailBuilder.EditCategory("Pivot", FText::GetEmpty(), ECategoryPriority::Uncommon);
+	PivotCategoryBuilder.AddProperty(PivotTypeProperty);
+	IDetailPropertyRow& PivotActorPropertyRow = PivotCategoryBuilder.AddProperty(PivotActorProperty);
+	
+	TSharedPtr<SWidget> NameWidget;
+	TSharedPtr<SWidget> ValueWidget;
+	FDetailWidgetRow Row;
+	PivotActorPropertyRow.GetDefaultWidgets(NameWidget, ValueWidget, Row);
+
+	const bool bShowChildren = true;
+	PivotActorPropertyRow.CustomWidget(bShowChildren)
+		.NameContent()
+		.MinDesiredWidth(Row.NameWidget.MinWidth)
+		.MaxDesiredWidth(Row.NameWidget.MaxWidth)
+		[
+			NameWidget.ToSharedRef()
+		]
+		.ValueContent()
+		.MinDesiredWidth(Row.ValueWidget.MinWidth)
+		.MaxDesiredWidth(Row.ValueWidget.MaxWidth)
+		[
+			SNew(SComboBox<AActor*>)
+			.OptionsSource(&PivotActors)
+			.OnGenerateWidget(this, &FNewLevelInstanceParamsDetails::OnGeneratePivotActorWidget)
+			.OnSelectionChanged(this, &FNewLevelInstanceParamsDetails::OnSelectedPivotActorChanged)
+			.IsEnabled(this, &FNewLevelInstanceParamsDetails::IsPivotActorSelectionEnabled)
+			[
+				SNew(STextBlock)
+				.Text(this, &FNewLevelInstanceParamsDetails::GetSelectedPivotActorText)
+				.Font(DetailBuilder.GetDetailFont())
+			]
+		];
 }
 
 #undef LOCTEXT_NAMESPACE
diff --git a/Engine/Source/Editor/LevelInstanceEditor/Private/SNewLevelInstanceDialog.h b/Engine/Source/Editor/LevelInstanceEditor/Private/SNewLevelInstanceDialog.h
--- a/Engine/Source/Editor/LevelInstanceEditor/Private/SNewLevelInstanceDialog.h
+++ b/Engine/Source/Editor/LevelInstanceEditor/Private/SNewLevelInstanceDialog.h
@@ -10,10 +10,10 @@
 #include "Widgets/SWindow.h"
 #include "Widgets/DeclarativeSyntaxSupport.h"
 #include "Widgets/SCompoundWidget.h"
-
-enum class ELevelInstanceCreationType : uint8;
-
-enum class ELevelInstancePivotType : uint8;
+#include "IDetailCustomization.h"
+#include "DetailLayoutBuilder.h"
+#include "Templates/SharedPointer.h"
+#include "LevelInstance/LevelInstanceTypes.h"
 
 //////////////////////////////////////////////////////////////////////////
 // SNewLevelInstanceDialog
@@ -41,42 +41,46 @@ public:
 	static const FVector2D DEFAULT_WINDOW_SIZE;
 
 	bool ClickedOk() const { return bClickedOk; }
-	ELevelInstanceCreationType GetCreationType() const { return SelectedCreationType; }
-	ELevelInstancePivotType GetPivotType() const { return SelectedPivotType; }
-	AActor* GetPivotActor() const { return SelectedPivotActor; }
+	
+	const FNewLevelInstanceParams& GetCreationParams() const { return CreationParams; }
 
 private:
-	int32 GetSelectedCreationType() const;
-	void OnSelectedCreationTypeChanged(int32 NewValue, ESelectInfo::Type SelectionType);
-
-	int32 GetSelectedPivotType() const;
-	void OnSelectedPivotTypeChanged(int32 NewValue, ESelectInfo::Type SelectionType);
-		
 	FReply OnOkClicked();
 	bool IsOkEnabled() const;
 
 	FReply OnCancelClicked();
 
-	TSharedRef<SWidget> OnGeneratePivotActorWidget(AActor* Actor) const;
-	FText GetSelectedPivotActorText() const;
-	void OnSelectedPivotActorChanged(AActor* NewValue, ESelectInfo::Type SelectionType);
-	bool IsPivotActorSelectionEnabled() const;
-
 	/** Pointer to the parent window, so we know to destroy it when done */
 	TWeakPtr<SWindow> ParentWindowPtr;
 
-	/** The type of LevelInstance to create */
-	ELevelInstanceCreationType SelectedCreationType;
+	/** Dialog Result */
+	FNewLevelInstanceParams CreationParams;
+
+	/** Dialog Result */
+	bool bClickedOk;
+};
+
+class FNewLevelInstanceParamsDetails : public IDetailCustomization
+{
+public:
 
-	/** The type of Pivot */
-	ELevelInstancePivotType SelectedPivotType;
+	FNewLevelInstanceParamsDetails(TArray<AActor*> InPivotActors)
+		: CreationParams(nullptr), PivotActors(InPivotActors)
+	{};
 
-	/** Pivot Actor */
-	AActor* SelectedPivotActor = nullptr;
+	static TSharedRef<IDetailCustomization> MakeInstance(TArray<AActor*> InPivotActors)
+	{
+		return MakeShareable(new FNewLevelInstanceParamsDetails(InPivotActors));
+	}
 
-	/** Actor List */
-	TArray<AActor*> PivotActors;
+	virtual void CustomizeDetails(IDetailLayoutBuilder& DetailBuilder) override;
 
-	/** Dialog Result */
-	bool bClickedOk;
+	TSharedRef<SWidget> OnGeneratePivotActorWidget(AActor* Actor) const;
+	FText GetSelectedPivotActorText() const;
+	void OnSelectedPivotActorChanged(AActor* NewValue, ESelectInfo::Type SelectionType);
+	bool IsPivotActorSelectionEnabled() const;
+
+protected:
+	FNewLevelInstanceParams* CreationParams;
+	TArray<AActor*> PivotActors;
 };
\ No newline at end of file
diff --git a/Engine/Source/Runtime/Engine/Private/LevelInstance/LevelInstanceSubsystem.cpp b/Engine/Source/Runtime/Engine/Private/LevelInstance/LevelInstanceSubsystem.cpp
--- a/Engine/Source/Runtime/Engine/Private/LevelInstance/LevelInstanceSubsystem.cpp
+++ b/Engine/Source/Runtime/Engine/Private/LevelInstance/LevelInstanceSubsystem.cpp
@@ -638,7 +638,7 @@ bool ULevelInstanceSubsystem::MoveActorsTo(ALevelInstance* LevelInstanceActor, c
 	return MoveActorsToLevel(ActorsToMove, LevelInstanceLevel);
 }
 
-ALevelInstance* ULevelInstanceSubsystem::CreateLevelInstanceFrom(const TArray<AActor*>& ActorsToMove, ELevelInstanceCreationType CreationType, ELevelInstancePivotType PivotType, AActor* PivotActor, UWorld* TemplateWorld)
+ALevelInstance* ULevelInstanceSubsystem::CreateLevelInstanceFrom(const TArray<AActor*>& ActorsToMove, const FNewLevelInstanceParams& CreationParams)
 {
 	ULevel* CurrentLevel = GetWorld()->GetCurrentLevel();
 		
@@ -664,21 +664,21 @@ ALevelInstance* ULevelInstanceSubsystem::CreateLevelInstanceFrom(const TArray<AA
 	}
 
 	FVector LevelInstanceLocation;
-	if (PivotType == ELevelInstancePivotType::Actor)
+	if (CreationParams.PivotType == ELevelInstancePivotType::Actor)
 	{
-		check(PivotActor);
-		LevelInstanceLocation = PivotActor->GetActorLocation();
+		check(CreationParams.PivotActor);
+		LevelInstanceLocation = CreationParams.PivotActor->GetActorLocation();
 	}
 	else
 	{
 		LevelInstanceLocation = ActorLocationBox.GetCenter();
-		if (PivotType == ELevelInstancePivotType::CenterMinZ)
+		if (CreationParams.PivotType == ELevelInstancePivotType::CenterMinZ)
 		{
 			LevelInstanceLocation.Z = ActorLocationBox.Min.Z;
 		}
 	}
 		
-	ULevelStreamingLevelInstanceEditor* LevelStreaming = StaticCast<ULevelStreamingLevelInstanceEditor*>(EditorLevelUtils::CreateNewStreamingLevelForWorld(*GetWorld(), ULevelStreamingLevelInstanceEditor::StaticClass(), TEXT(""), false, TemplateWorld));
+	ULevelStreamingLevelInstanceEditor* LevelStreaming = StaticCast<ULevelStreamingLevelInstanceEditor*>(EditorLevelUtils::CreateNewStreamingLevelForWorld(*GetWorld(), ULevelStreamingLevelInstanceEditor::StaticClass(), TEXT(""), false, CreationParams.TemplateWorld));
 	if (!LevelStreaming)
 	{
 		UE_LOG(LogLevelInstance, Warning, TEXT("Failed to create new Level Instance level"));
@@ -699,6 +699,13 @@ ALevelInstance* ULevelInstanceSubsystem::CreateLevelInstanceFrom(const TArray<AA
 		return nullptr;
 	}
 	
+	// Convert to OFPA
+	if (CreationParams.bExternalActors)
+	{
+		LoadedLevel->ConvertAllActorsToPackaging(true);
+		LoadedLevel->bUseExternalActors = true;
+	}
+
 	// Take all actors out of any folders they may have been in since we don't support folders inside of level instances
 	for (AActor* Actor : LoadedLevel->Actors)
 	{
@@ -713,15 +720,15 @@ ALevelInstance* ULevelInstanceSubsystem::CreateLevelInstanceFrom(const TArray<AA
 	ALevelInstance* NewLevelInstanceActor = nullptr;
 	TSoftObjectPtr<UWorld> WorldPtr(LoadedLevel->GetTypedOuter<UWorld>());
 	
-	if (CreationType == ELevelInstanceCreationType::LevelInstance)
+	if (CreationParams.Type == ELevelInstanceCreationType::LevelInstance)
 	{
 		NewLevelInstanceActor = GetWorld()->SpawnActor<ALevelInstance>(ALevelInstance::StaticClass(), SpawnParams);
 	}
-	else if (CreationType == ELevelInstanceCreationType::PackedLevelInstance)
+	else if (CreationParams.Type == ELevelInstanceCreationType::PackedLevelInstance)
 	{
 		NewLevelInstanceActor = GetWorld()->SpawnActor<APackedLevelInstance>(APackedLevelInstance::StaticClass(), SpawnParams);
 	}
-	else if (CreationType == ELevelInstanceCreationType::PackedLevelInstanceBlueprint)
+	else if (CreationParams.Type == ELevelInstanceCreationType::PackedLevelInstanceBlueprint)
 	{
 		int32 LastSlashIndex = 0;
 		FString LongPackageName = WorldPtr.GetLongPackageName();
diff --git a/Engine/Source/Runtime/Engine/Public/LevelInstance/LevelInstanceSubsystem.h b/Engine/Source/Runtime/Engine/Public/LevelInstance/LevelInstanceSubsystem.h
--- a/Engine/Source/Runtime/Engine/Public/LevelInstance/LevelInstanceSubsystem.h
+++ b/Engine/Source/Runtime/Engine/Public/LevelInstance/LevelInstanceSubsystem.h
@@ -73,7 +73,7 @@ public:
 
 	bool SetCurrent(ALevelInstance* LevelInstanceActor) const;
 	bool IsCurrent(const ALevelInstance* LevelInstanceActor) const;
-	ALevelInstance* CreateLevelInstanceFrom(const TArray<AActor*>& ActorsToMove, ELevelInstanceCreationType CreationType, ELevelInstancePivotType PivotType, AActor* PivotActor = nullptr, UWorld* TemplateWorld = nullptr);
+	ALevelInstance* CreateLevelInstanceFrom(const TArray<AActor*>& ActorsToMove, const FNewLevelInstanceParams& CreationParams);
 	bool MoveActorsToLevel(const TArray<AActor*>& ActorsToRemove, ULevel* DestinationLevel) const;
 	bool MoveActorsTo(ALevelInstance* LevelInstanceActor, const TArray<AActor*>& ActorsToMove);
 	bool BreakLevelInstance(ALevelInstance* LevelInstanceActor, uint32 Levels = 1);
diff --git a/Engine/Source/Runtime/Engine/Public/LevelInstance/LevelInstanceTypes.h b/Engine/Source/Runtime/Engine/Public/LevelInstance/LevelInstanceTypes.h
--- a/Engine/Source/Runtime/Engine/Public/LevelInstance/LevelInstanceTypes.h
+++ b/Engine/Source/Runtime/Engine/Public/LevelInstance/LevelInstanceTypes.h
@@ -2,6 +2,9 @@
 
 #pragma once
 
+#include "CoreMinimal.h"
+#include "LevelInstanceTypes.generated.h"
+
 // FLevelInstanceID is a runtime unique id that is computed from the Hash of LevelInstance Actor Guid and all its ancestor LevelInstance Actor Guids.
 // Resulting in a different ID for all instances whether they load the same level or not.
 using FLevelInstanceID = uint32;
@@ -22,4 +25,25 @@ enum class ELevelInstancePivotType : uint8
 	CenterMinZ,
 	Center,
 	Actor
+};
+
+USTRUCT()
+struct FNewLevelInstanceParams
+{
+	GENERATED_USTRUCT_BODY()
+			
+	UPROPERTY(EditAnywhere, Category = Default)
+	ELevelInstanceCreationType Type = ELevelInstanceCreationType::LevelInstance;
+
+	UPROPERTY(EditAnywhere, Category = Default)
+	bool bExternalActors = true;
+
+	UPROPERTY(EditAnywhere, Category = Pivot)
+	ELevelInstancePivotType PivotType = ELevelInstancePivotType::CenterMinZ;
+
+	UPROPERTY(EditAnywhere, Category = Pivot)
+	AActor* PivotActor = nullptr;
+
+	UPROPERTY()
+	UWorld* TemplateWorld = nullptr;
 };
\ No newline at end of file
