LevelInstance: Added auto sizing for New Level Instance Dialog.

diff --git a/Engine/Source/Editor/LevelInstanceEditor/Private/SNewLevelInstanceDialog.cpp b/Engine/Source/Editor/LevelInstanceEditor/Private/SNewLevelInstanceDialog.cpp
--- a/Engine/Source/Editor/LevelInstanceEditor/Private/SNewLevelInstanceDialog.cpp
+++ b/Engine/Source/Editor/LevelInstanceEditor/Private/SNewLevelInstanceDialog.cpp
@@ -59,9 +59,14 @@ void SNewLevelInstanceDialog::Construct(const FArguments& InArgs)
 		[
 			SNew(SVerticalBox)
 			+ SVerticalBox::Slot()
-			.MaxHeight(250.0f)
+			.AutoHeight()
 			[
-				StructureDetailsView->GetWidget()->AsShared()
+				SNew(SHorizontalBox)
+				+ SHorizontalBox::Slot()
+				.AutoWidth()
+				[
+					StructureDetailsView->GetWidget()->AsShared()
+				]
 			]
 			+ SVerticalBox::Slot()
 			.AutoHeight()
