Fixed non-unity build, added missing include for FReply

diff --git a/Engine/Source/Editor/LevelInstanceEditor/Private/LevelInstancePivotDetails.h b/Engine/Source/Editor/LevelInstanceEditor/Private/LevelInstancePivotDetails.h
--- a/Engine/Source/Editor/LevelInstanceEditor/Private/LevelInstancePivotDetails.h
+++ b/Engine/Source/Editor/LevelInstanceEditor/Private/LevelInstancePivotDetails.h
@@ -6,6 +6,7 @@
 #include "Types/SlateEnums.h"
 #include "IDetailCustomization.h"
 #include "LevelInstance/LevelInstanceTypes.h"
+#include "Input/Reply.h"
 
 class ALevelInstancePivot;
 
