LevelInstance: Add Discard menu

diff --git a/Engine/Source/Editor/LevelInstanceEditor/Private/LevelInstanceEditorModule.cpp b/Engine/Source/Editor/LevelInstanceEditor/Private/LevelInstanceEditorModule.cpp
--- a/Engine/Source/Editor/LevelInstanceEditor/Private/LevelInstanceEditorModule.cpp
+++ b/Engine/Source/Editor/LevelInstanceEditor/Private/LevelInstanceEditorModule.cpp
@@ -71,7 +71,7 @@ namespace LevelInstanceMenuUtils
 		}
 	}
 
-	void CreateCommitSubMenu(UToolMenu* Menu, TArray<ALevelInstance*> LevelInstanceEdits, ALevelInstance* ContextLevelInstance)
+	void CreateCommitSubMenu(UToolMenu* Menu, TArray<ALevelInstance*> LevelInstanceEdits, ALevelInstance* ContextLevelInstance, bool bDiscard)
 	{
 		FText OtherSectionLabel = LOCTEXT("LevelInstanceOtherCommitSection", "Other(s)");
 		FToolMenuSection* Section = &Menu->AddSection("LevelInstanceContextCommitSection", ContextLevelInstance != nullptr? LOCTEXT("LevelInstanceContextCommitSection", "Context") : OtherSectionLabel);
@@ -81,9 +81,16 @@ namespace LevelInstanceMenuUtils
 			const bool bCanCommit = LevelInstanceActor->CanCommit(&EntryDesc);
 
 			FToolUIAction LevelInstanceEditAction;
-			LevelInstanceEditAction.ExecuteAction.BindLambda([LevelInstanceActor](const FToolMenuContext&)
+			LevelInstanceEditAction.ExecuteAction.BindLambda([bDiscard, LevelInstanceActor](const FToolMenuContext&)
 			{
-				LevelInstanceActor->Commit();
+				if (bDiscard)
+				{
+					LevelInstanceActor->Discard();
+				}
+				else
+				{
+					LevelInstanceActor->Commit();
+				}
 			});
 			LevelInstanceEditAction.CanExecuteAction.BindLambda([bCanCommit](const FToolMenuContext&)
 			{
@@ -220,7 +227,13 @@ namespace LevelInstanceMenuUtils
 				"CommitLevelInstances",
 				LOCTEXT("CommitLevelInstances", "Commit"),
 				TAttribute<FText>(),
-				FNewToolMenuDelegate::CreateStatic(&CreateCommitSubMenu, MoveTemp(LevelInstanceEdits), ContextLevelInstance)
+				FNewToolMenuDelegate::CreateStatic(&CreateCommitSubMenu, LevelInstanceEdits, ContextLevelInstance, /*bDiscard=*/ false)
+			);
+			Section.AddSubMenu(
+				"DiscardLevelInstances",
+				LOCTEXT("DiscardLevelInstances", "Discard"),
+				TAttribute<FText>(),
+				FNewToolMenuDelegate::CreateStatic(&CreateCommitSubMenu, LevelInstanceEdits, ContextLevelInstance, /*bDiscard=*/ true)
 			);
 		}
 	}
