LevelInstance: add a CanExecuteAction lambda to Move Selection To

diff --git a/Engine/Source/Editor/LevelInstanceEditor/Private/LevelInstanceEditorModule.cpp b/Engine/Source/Editor/LevelInstanceEditor/Private/LevelInstanceEditorModule.cpp
--- a/Engine/Source/Editor/LevelInstanceEditor/Private/LevelInstanceEditorModule.cpp
+++ b/Engine/Source/Editor/LevelInstanceEditor/Private/LevelInstanceEditorModule.cpp
@@ -150,16 +150,36 @@ namespace LevelInstanceMenuUtils
 	void CreateMoveSelectionToSubMenu(UToolMenu* Menu, TArray<ALevelInstance*> LevelInstanceEdits)
 	{
 		FToolMenuSection* Section = &Menu->AddSection(NAME_None);
-		for (ALevelInstance* LevelInstanceActor : LevelInstanceEdits)
+		if (ULevelInstanceSubsystem* LevelInstanceSubsystem = GEditor->GetEditorWorldContext().World()->GetSubsystem<ULevelInstanceSubsystem>())
 		{
-			FToolUIAction LevelInstanceMoveSelectionAction;
-			LevelInstanceMoveSelectionAction.ExecuteAction.BindLambda([LevelInstanceActor](const FToolMenuContext&)
+			for (ALevelInstance* LevelInstanceActor : LevelInstanceEdits)
 			{
-				MoveSelectionToLevelInstance(LevelInstanceActor);
-			});
+				FToolUIAction LevelInstanceMoveSelectionAction;
 
-			const FText EntryLabel = FText::Format(LOCTEXT("LevelInstanceName", "{0}:{1}"), FText::FromString(LevelInstanceActor->GetName()), FText::FromString(LevelInstanceActor->GetWorldAssetPackage()));
-			Section->AddMenuEntry(NAME_None, EntryLabel, TAttribute<FText>(), FSlateIcon(), LevelInstanceMoveSelectionAction);
+				LevelInstanceMoveSelectionAction.CanExecuteAction.BindLambda([LevelInstanceActor, LevelInstanceSubsystem](const FToolMenuContext& MenuContext)
+					{
+						for (FSelectionIterator It(GEditor->GetSelectedActorIterator()); It; ++It)
+						{
+							if (AActor* Actor = Cast<AActor>(*It))
+							{
+								if (Actor->GetLevel() == LevelInstanceSubsystem->GetLevelInstanceLevel(LevelInstanceActor))
+								{
+									return false;
+								}
+							}
+						}
+					
+						return GEditor->GetSelectedActorCount() > 0;
+					});
+
+				LevelInstanceMoveSelectionAction.ExecuteAction.BindLambda([LevelInstanceActor](const FToolMenuContext&)
+					{
+						MoveSelectionToLevelInstance(LevelInstanceActor);
+					});
+
+				const FText EntryLabel = FText::Format(LOCTEXT("LevelInstanceName", "{0}:{1}"), FText::FromString(LevelInstanceActor->GetName()), FText::FromString(LevelInstanceActor->GetWorldAssetPackage()));
+				Section->AddMenuEntry(NAME_None, EntryLabel, TAttribute<FText>(), FSlateIcon(), LevelInstanceMoveSelectionAction);
+			}
 		}
 	}
 
