LevelInstance: Use ActorLabel in menus

diff --git a/Engine/Source/Editor/LevelInstanceEditor/Private/LevelInstanceEditorModule.cpp b/Engine/Source/Editor/LevelInstanceEditor/Private/LevelInstanceEditorModule.cpp
--- a/Engine/Source/Editor/LevelInstanceEditor/Private/LevelInstanceEditorModule.cpp
+++ b/Engine/Source/Editor/LevelInstanceEditor/Private/LevelInstanceEditorModule.cpp
@@ -66,7 +66,7 @@ namespace LevelInstanceMenuUtils
 				return bCanEdit;
 			});
 
-			const FText EntryLabel = FText::Format(LOCTEXT("LevelInstanceName", "{0}:{1}"), FText::FromString(LevelInstanceActor->GetName()), FText::FromString(LevelInstanceActor->GetWorldAssetPackage()));
+			const FText EntryLabel = FText::Format(LOCTEXT("LevelInstanceName", "{0}:{1}"), FText::FromString(LevelInstanceActor->GetActorLabel()), FText::FromString(LevelInstanceActor->GetWorldAssetPackage()));
 			Section.AddMenuEntry(NAME_None, EntryLabel, EntryDesc, FSlateIcon(), LevelInstanceEditAction);
 		}
 	}
@@ -97,7 +97,7 @@ namespace LevelInstanceMenuUtils
 				return bCanCommit;
 			});
 
-			const FText EntryLabel = FText::Format(LOCTEXT("LevelInstanceName", "{0}:{1}"), FText::FromString(LevelInstanceActor->GetName()), FText::FromString(LevelInstanceActor->GetWorldAssetPackage()));
+			const FText EntryLabel = FText::Format(LOCTEXT("LevelInstanceName", "{0}:{1}"), FText::FromString(LevelInstanceActor->GetActorLabel()), FText::FromString(LevelInstanceActor->GetWorldAssetPackage()));
 			Section->AddMenuEntry(NAME_None, EntryLabel, EntryDesc, FSlateIcon(), LevelInstanceEditAction);
 
 			if (ContextLevelInstance == LevelInstanceActor && LevelInstanceEdits.Num() > 1)
@@ -119,7 +119,7 @@ namespace LevelInstanceMenuUtils
 				LevelInstanceActor->SetCurrent();
 			});
 		
-			const FText EntryLabel = FText::Format(LOCTEXT("LevelInstanceName", "{0}:{1}"), FText::FromString(LevelInstanceActor->GetName()), FText::FromString(LevelInstanceActor->GetWorldAssetPackage()));
+			const FText EntryLabel = FText::Format(LOCTEXT("LevelInstanceName", "{0}:{1}"), FText::FromString(LevelInstanceActor->GetActorLabel()), FText::FromString(LevelInstanceActor->GetWorldAssetPackage()));
 			Section->AddMenuEntry(NAME_None, EntryLabel, TAttribute<FText>(), FSlateIcon(), LevelInstanceSetCurrentAction);
 
 			if (LevelInstanceActor == ContextLevelInstance && LevelInstanceEdits.Num() > 1)
@@ -177,7 +177,7 @@ namespace LevelInstanceMenuUtils
 						MoveSelectionToLevelInstance(LevelInstanceActor);
 					});
 
-				const FText EntryLabel = FText::Format(LOCTEXT("LevelInstanceName", "{0}:{1}"), FText::FromString(LevelInstanceActor->GetName()), FText::FromString(LevelInstanceActor->GetWorldAssetPackage()));
+				const FText EntryLabel = FText::Format(LOCTEXT("LevelInstanceName", "{0}:{1}"), FText::FromString(LevelInstanceActor->GetActorLabel()), FText::FromString(LevelInstanceActor->GetWorldAssetPackage()));
 				Section->AddMenuEntry(NAME_None, EntryLabel, TAttribute<FText>(), FSlateIcon(), LevelInstanceMoveSelectionAction);
 			}
 		}
