Non-unity fix for LevelInstancePivotDetails.cpp

diff --git a/Engine/Source/Editor/LevelInstanceEditor/Private/LevelInstancePivotDetails.cpp b/Engine/Source/Editor/LevelInstanceEditor/Private/LevelInstancePivotDetails.cpp
--- a/Engine/Source/Editor/LevelInstanceEditor/Private/LevelInstancePivotDetails.cpp
+++ b/Engine/Source/Editor/LevelInstanceEditor/Private/LevelInstancePivotDetails.cpp
@@ -8,7 +8,7 @@
 #include "Widgets/Input/SMultiLineEditableTextBox.h"
 #include "SEnumCombo.h"
 #include "ScopedTransaction.h"
-
+#include "Engine/Brush.h"
 #include "Engine/World.h"
 #include "LevelInstance/LevelInstanceActor.h"
 #include "LevelInstance/LevelInstanceEditorPivotActor.h"
