Moving level instance section of the level editor context menu below Actor Actions, added section heading for Merge Actors

diff --git a/Engine/Source/Editor/LevelEditor/Private/LevelEditorContextMenu.cpp b/Engine/Source/Editor/LevelEditor/Private/LevelEditorContextMenu.cpp
--- a/Engine/Source/Editor/LevelEditor/Private/LevelEditorContextMenu.cpp
+++ b/Engine/Source/Editor/LevelEditor/Private/LevelEditorContextMenu.cpp
@@ -321,7 +321,7 @@ void FLevelEditorContextMenu::RegisterActorContextMenu()
 
 		if (LevelEditorContext->ContextType == ELevelEditorMenuContext::Viewport || LevelEditorContext->ContextType == ELevelEditorMenuContext::SceneOutliner)
 		{
-			FToolMenuSection& Section = InMenu->AddSection("OpenMergeActor");
+			FToolMenuSection& Section = InMenu->AddSection("OpenMergeActor", LOCTEXT("MergeActors","Merge Actors"));
 
 			Section.AddSubMenu(
 				"MergeActorsSubMenu",
diff --git a/Engine/Source/Editor/LevelInstanceEditor/Private/LevelInstanceEditorModule.cpp b/Engine/Source/Editor/LevelInstanceEditor/Private/LevelInstanceEditorModule.cpp
--- a/Engine/Source/Editor/LevelInstanceEditor/Private/LevelInstanceEditorModule.cpp
+++ b/Engine/Source/Editor/LevelInstanceEditor/Private/LevelInstanceEditorModule.cpp
@@ -43,8 +43,14 @@ namespace LevelInstanceMenuUtils
 {
 	FToolMenuSection& CreateLevelInstanceSection(UToolMenu* Menu)
 	{
-		FToolMenuSection& Section = Menu->FindOrAddSection(TEXT("LevelInstance"));
-		Section.Label = LOCTEXT("LevelInstance", "Level Instance");
+		const FName LevelInstanceSectionName = TEXT("LevelInstance");
+		FToolMenuSection* SectionPtr = Menu->FindSection(LevelInstanceSectionName);
+		if (!SectionPtr)
+		{
+			FToolMenuInsert InsertPosition("ActorControl", EToolMenuInsertType::After);
+			SectionPtr = &(Menu->AddSection(LevelInstanceSectionName, LOCTEXT("LevelInstance", "Level Instance"), InsertPosition));
+		}
+		FToolMenuSection& Section = *SectionPtr;
 		return Section;
 	}
 
