LevelInstance: Exit Mode on PIE (allow saving)

diff --git a/Engine/Source/Editor/LevelInstanceEditor/Private/LevelInstanceEditorMode.cpp b/Engine/Source/Editor/LevelInstanceEditor/Private/LevelInstanceEditorMode.cpp
--- a/Engine/Source/Editor/LevelInstanceEditor/Private/LevelInstanceEditorMode.cpp
+++ b/Engine/Source/Editor/LevelInstanceEditor/Private/LevelInstanceEditorMode.cpp
@@ -24,12 +24,22 @@ ULevelInstanceEditorMode::ULevelInstanceEditorMode()
 		false);
 
 	bContextRestriction = true;
+
+	FEditorDelegates::PreBeginPIE.AddUObject(this, &ULevelInstanceEditorMode::OnPreBeginPIE);
 }
 
 ULevelInstanceEditorMode::~ULevelInstanceEditorMode()
 {
 }
 
+void ULevelInstanceEditorMode::OnPreBeginPIE(bool bSimulate)
+{
+	if (GLevelEditorModeTools().IsModeActive(ULevelInstanceEditorMode::EM_LevelInstanceEditorModeId))
+	{
+		GLevelEditorModeTools().DeactivateMode(ULevelInstanceEditorMode::EM_LevelInstanceEditorModeId);
+	}
+}
+
 void ULevelInstanceEditorMode::UpdateEngineShowFlags()
 {
 	for (FLevelEditorViewportClient* LevelVC : GEditor->GetLevelViewportClients())
diff --git a/Engine/Source/Editor/LevelInstanceEditor/Private/LevelInstanceEditorMode.h b/Engine/Source/Editor/LevelInstanceEditor/Private/LevelInstanceEditorMode.h
--- a/Engine/Source/Editor/LevelInstanceEditor/Private/LevelInstanceEditorMode.h
+++ b/Engine/Source/Editor/LevelInstanceEditor/Private/LevelInstanceEditorMode.h
@@ -32,6 +32,7 @@ public:
 	// End UEdMode
 
 private:
+	void OnPreBeginPIE(bool bSimulate);
 	void UpdateEngineShowFlags();
 	virtual void BindCommands() override;
 
