LevelInstance: avoid Edit sub-menu when only one Level can be edited

diff --git a/Engine/Source/Editor/LevelInstanceEditor/Private/LevelInstanceEditorModule.cpp b/Engine/Source/Editor/LevelInstanceEditor/Private/LevelInstanceEditorModule.cpp
--- a/Engine/Source/Editor/LevelInstanceEditor/Private/LevelInstanceEditorModule.cpp
+++ b/Engine/Source/Editor/LevelInstanceEditor/Private/LevelInstanceEditorModule.cpp
@@ -58,26 +58,35 @@ namespace LevelInstanceMenuUtils
 		return Section;
 	}
 
-	void CreateEditSubMenu(UToolMenu* Menu, TArray<ALevelInstance*> LevelInstanceHierarchy, AActor* ContextActor)
+	void CreateEditMenuEntry(FToolMenuSection& Section, ALevelInstance* LevelInstance, AActor* ContextActor, bool bSingleEntry)
 	{
-		FToolMenuSection& Section = Menu->AddSection(NAME_None, LOCTEXT("LevelInstanceContextEditSection", "Context"));
-		for (ALevelInstance* LevelInstanceActor : LevelInstanceHierarchy)
+		FToolUIAction LevelInstanceEditAction;
+		FText EntryDesc;
+		const bool bCanEdit = LevelInstance->CanEdit(&EntryDesc);
+
+		LevelInstanceEditAction.ExecuteAction.BindLambda([LevelInstance, ContextActor](const FToolMenuContext&)
+		{
+			LevelInstance->Edit(ContextActor);
+		});
+		LevelInstanceEditAction.CanExecuteAction.BindLambda([bCanEdit](const FToolMenuContext&)
 		{
-			FToolUIAction LevelInstanceEditAction;
-			FText EntryDesc = LOCTEXT("LevelInstanceEditSubMenuEntry","");
-			const bool bCanEdit = LevelInstanceActor->CanEdit(&EntryDesc);
+			return bCanEdit;
+		});
 
-			LevelInstanceEditAction.ExecuteAction.BindLambda([LevelInstanceActor, ContextActor](const FToolMenuContext&)
-			{
-				LevelInstanceActor->Edit(ContextActor);
-			});
-			LevelInstanceEditAction.CanExecuteAction.BindLambda([bCanEdit](const FToolMenuContext&)
-			{ 
-				return bCanEdit;
-			});
+		FText EntryLabel = bSingleEntry ? LOCTEXT("EditLevelInstances", "Edit") : FText::FromString(LevelInstance->GetWorldAsset().GetAssetName());
+		if (bCanEdit)
+		{
+			EntryDesc = FText::Format(LOCTEXT("LevelInstanceName", "{0}:{1}"), FText::FromString(LevelInstance->GetActorLabel()), FText::FromString(LevelInstance->GetWorldAssetPackage()));
+		}
+		Section.AddMenuEntry(NAME_None, EntryLabel, EntryDesc, FSlateIcon(), LevelInstanceEditAction);
+	}
 
-			const FText EntryLabel = FText::Format(LOCTEXT("LevelInstanceName", "{0}:{1}"), FText::FromString(LevelInstanceActor->GetActorLabel()), FText::FromString(LevelInstanceActor->GetWorldAssetPackage()));
-			Section.AddMenuEntry(NAME_None, EntryLabel, EntryDesc, FSlateIcon(), LevelInstanceEditAction);
+	void CreateEditSubMenu(UToolMenu* Menu, TArray<ALevelInstance*> LevelInstanceHierarchy, AActor* ContextActor)
+	{
+		FToolMenuSection& Section = Menu->AddSection(NAME_None, LOCTEXT("LevelInstanceContextEditSection", "Context"));
+		for (ALevelInstance* LevelInstance : LevelInstanceHierarchy)
+		{
+			CreateEditMenuEntry(Section, LevelInstance, ContextActor, false);
 		}
 	}
 		
@@ -110,7 +119,13 @@ namespace LevelInstanceMenuUtils
 				return true;
 			});
 
-			if (LevelInstanceHierarchy.Num() > 0)
+			// Don't create sub menu if only one Level Instance is available to edit
+			if (LevelInstanceHierarchy.Num() == 1)
+			{
+				FToolMenuSection& Section = CreateLevelInstanceSection(Menu);
+				CreateEditMenuEntry(Section, LevelInstanceHierarchy[0], ContextActor, true);
+			}
+			else if(LevelInstanceHierarchy.Num() > 1)
 			{
 				FToolMenuSection& Section = CreateLevelInstanceSection(Menu);
 				Section.AddSubMenu(
diff --git a/Engine/Source/Runtime/Engine/Private/LevelInstance/LevelInstanceSubsystem.cpp b/Engine/Source/Runtime/Engine/Private/LevelInstance/LevelInstanceSubsystem.cpp
--- a/Engine/Source/Runtime/Engine/Private/LevelInstance/LevelInstanceSubsystem.cpp
+++ b/Engine/Source/Runtime/Engine/Private/LevelInstance/LevelInstanceSubsystem.cpp
@@ -1492,7 +1492,7 @@ bool ULevelInstanceSubsystem::CanEditLevelInstance(const ALevelInstance* LevelIn
 		{
 			if (OutReason)
 			{
-				*OutReason = LOCTEXT("CanEditPartitionedLevelInstance", "Can't edit partitioned Level Instance");
+				*OutReason = FText::Format(LOCTEXT("CanEditPartitionedLevelInstance", "Can't edit partitioned Level Instance ({0})."), FText::FromString(LevelInstanceActor->GetWorldAssetPackage()));
 			}
 			return false;
 		}
@@ -1504,7 +1504,7 @@ bool ULevelInstanceSubsystem::CanEditLevelInstance(const ALevelInstance* LevelIn
 		{
 			if (OutReason)
 			{
-				*OutReason = LOCTEXT("CanEditLevelInstanceAlreadyBeingEdited", "Level Instance already being edited");
+				*OutReason = FText::Format(LOCTEXT("CanEditLevelInstanceAlreadyBeingEdited", "Level Instance already being edited ({0})."), FText::FromString(LevelInstanceActor->GetWorldAssetPackage()));
 			}
 			return false;
 		}
@@ -1513,7 +1513,7 @@ bool ULevelInstanceSubsystem::CanEditLevelInstance(const ALevelInstance* LevelIn
 		{
 			if (OutReason)
 			{
-				*OutReason = LOCTEXT("CanEditLevelInstanceDirtyEdit", "Current Level Instance has unsaved changes and needs to be committed first.");
+				*OutReason = FText::Format(LOCTEXT("CanEditLevelInstanceDirtyEdit", "Current Level Instance has unsaved changes and needs to be committed first ({0})."), FText::FromString(GetEditingLevelInstance()->GetWorldAssetPackage()));
 			}
 			return false;
 		}
@@ -1523,7 +1523,7 @@ bool ULevelInstanceSubsystem::CanEditLevelInstance(const ALevelInstance* LevelIn
 	{
 		if (OutReason)
 		{
-			*OutReason = LOCTEXT("CanEditLevelInstanceDirtyInvalid", "Level Instance path is invalid");
+			*OutReason = LOCTEXT("CanEditLevelInstanceDirtyInvalid", "Level Instance path is invalid.");
 		}
 		return false;
 	}
@@ -1532,7 +1532,7 @@ bool ULevelInstanceSubsystem::CanEditLevelInstance(const ALevelInstance* LevelIn
 	{
 		if (OutReason)
 		{
-			*OutReason = LOCTEXT("CanEditLevelInstancePersistentLevel", "The Persistent level and the Level Instance are the same");
+			*OutReason = FText::Format(LOCTEXT("CanEditLevelInstancePersistentLevel", "The Persistent level and the Level Instance are the same ({0})."), FText::FromString(LevelInstanceActor->GetWorldAssetPackage()));
 		}
 		return false;
 	}
@@ -1541,7 +1541,7 @@ bool ULevelInstanceSubsystem::CanEditLevelInstance(const ALevelInstance* LevelIn
 	{
 		if (OutReason)
 		{
-			*OutReason = LOCTEXT("CanEditLevelInstanceAlreadyExists", "The same level was added to world outside of Level Instances");
+			*OutReason = FText::Format(LOCTEXT("CanEditLevelInstanceAlreadyExists", "The same level was added to world outside of Level Instances ({0})."), FText::FromString(LevelInstanceActor->GetWorldAssetPackage()));
 		}
 		return false;
 	}
