Level Instance & Packed Level Actor: Experimental Settings (on by default)

diff --git a/Engine/Source/Editor/LevelInstanceEditor/Private/LevelInstanceEditorModule.cpp b/Engine/Source/Editor/LevelInstanceEditor/Private/LevelInstanceEditorModule.cpp
--- a/Engine/Source/Editor/LevelInstanceEditor/Private/LevelInstanceEditorModule.cpp
+++ b/Engine/Source/Editor/LevelInstanceEditor/Private/LevelInstanceEditorModule.cpp
@@ -37,6 +37,7 @@
 #include "Widgets/SWindow.h"
 #include "SNewLevelInstanceDialog.h"
 #include "MessageLogModule.h"
+#include "Settings/EditorExperimentalSettings.h"
 
 IMPLEMENT_MODULE( FLevelInstanceEditorModule, LevelInstanceEditor );
 
@@ -44,6 +45,16 @@ IMPLEMENT_MODULE( FLevelInstanceEditorModule, LevelInstanceEditor );
 
 namespace LevelInstanceMenuUtils
 {
+	bool IsExperimentalSettingEnabled(ALevelInstance* LevelInstance)
+	{
+		if (LevelInstance->IsA<APackedLevelActor>() && !GetDefault<UEditorExperimentalSettings>()->bPackedLevelActor)
+		{
+			return false;
+		}
+
+		return GetDefault<UEditorExperimentalSettings>()->bLevelInstance;
+	}
+
 	FToolMenuSection& CreateLevelSection(UToolMenu* Menu)
 	{
 		const FName LevelSectionName = TEXT("Level");
@@ -113,7 +124,10 @@ namespace LevelInstanceMenuUtils
 			TArray<ALevelInstance*> LevelInstanceHierarchy;
 			LevelInstanceSubsystem->ForEachLevelInstanceAncestorsAndSelf(ContextActor, [&LevelInstanceHierarchy](ALevelInstance* AncestorLevelInstance)
 			{
-				LevelInstanceHierarchy.Add(AncestorLevelInstance);
+				if (IsExperimentalSettingEnabled(AncestorLevelInstance))
+				{
+					LevelInstanceHierarchy.Add(AncestorLevelInstance);
+				}
 				return true;
 			});
 
@@ -316,19 +330,25 @@ namespace LevelInstanceMenuUtils
 			{
 				FToolMenuSection& Section = ToolMenu->AddSection("ActorSelectionSectionName", LOCTEXT("ActorSelectionSectionLabel", "Actor Selection"));
 
-				Section.AddMenuEntry(
-					NAME_None,
-					FText::Format(LOCTEXT("CreateFromSelectionLabel", "Create {0}..."), StaticEnum<ELevelInstanceCreationType>()->GetDisplayNameTextByValue((int64)ELevelInstanceCreationType::LevelInstance)),
-					TAttribute<FText>(),
-					FSlateIcon(FAppStyle::GetAppStyleSetName(), "ClassIcon.LevelInstance"),
-					FExecuteAction::CreateStatic(&LevelInstanceMenuUtils::CreateLevelInstanceFromSelection, LevelInstanceSubsystem, ELevelInstanceCreationType::LevelInstance));
+				if (GetDefault<UEditorExperimentalSettings>()->bLevelInstance)
+				{
+					Section.AddMenuEntry(
+						NAME_None,
+						FText::Format(LOCTEXT("CreateFromSelectionLabel", "Create {0}..."), StaticEnum<ELevelInstanceCreationType>()->GetDisplayNameTextByValue((int64)ELevelInstanceCreationType::LevelInstance)),
+						TAttribute<FText>(),
+						FSlateIcon(FAppStyle::GetAppStyleSetName(), "ClassIcon.LevelInstance"),
+						FExecuteAction::CreateStatic(&LevelInstanceMenuUtils::CreateLevelInstanceFromSelection, LevelInstanceSubsystem, ELevelInstanceCreationType::LevelInstance));
+				}
 
-				Section.AddMenuEntry(
-					NAME_None,
-					FText::Format(LOCTEXT("CreateFromSelectionLabel", "Create {0}..."), StaticEnum<ELevelInstanceCreationType>()->GetDisplayNameTextByValue((int64)ELevelInstanceCreationType::PackedLevelActor)),
-					TAttribute<FText>(),
-					FSlateIcon(FAppStyle::GetAppStyleSetName(), "ClassIcon.PackedLevelActor"),
-					FExecuteAction::CreateStatic(&LevelInstanceMenuUtils::CreateLevelInstanceFromSelection, LevelInstanceSubsystem, ELevelInstanceCreationType::PackedLevelActor));
+				if (GetDefault<UEditorExperimentalSettings>()->bPackedLevelActor)
+				{
+					Section.AddMenuEntry(
+						NAME_None,
+						FText::Format(LOCTEXT("CreateFromSelectionLabel", "Create {0}..."), StaticEnum<ELevelInstanceCreationType>()->GetDisplayNameTextByValue((int64)ELevelInstanceCreationType::PackedLevelActor)),
+						TAttribute<FText>(),
+						FSlateIcon(FAppStyle::GetAppStyleSetName(), "ClassIcon.PackedLevelActor"),
+						FExecuteAction::CreateStatic(&LevelInstanceMenuUtils::CreateLevelInstanceFromSelection, LevelInstanceSubsystem, ELevelInstanceCreationType::PackedLevelActor));
+				}
 			}
 		}
 	}
@@ -352,7 +372,7 @@ namespace LevelInstanceMenuUtils
 			}
 		}
 
-		if (ContextLevelInstance)
+		if (ContextLevelInstance && IsExperimentalSettingEnabled(ContextLevelInstance))
 		{
 			FToolMenuSection& Section = CreateLevelSection(Menu);
 			FText EntryDesc = LOCTEXT("LevelInstanceEditSubMenuEntry", "");
@@ -449,7 +469,7 @@ namespace LevelInstanceMenuUtils
 					return true;
 				});
 
-			if (ContextLevelInstance && !ContextLevelInstance->IsEditing() && !LevelInstanceSubsystem->LevelInstanceHasLevelScriptBlueprint(ContextLevelInstance))
+			if (ContextLevelInstance && IsExperimentalSettingEnabled(ContextLevelInstance) && !ContextLevelInstance->IsEditing() && !LevelInstanceSubsystem->LevelInstanceHasLevelScriptBlueprint(ContextLevelInstance))
 			{
 				FToolMenuSection& Section = CreateLevelSection(Menu);
 
@@ -480,7 +500,7 @@ namespace LevelInstanceMenuUtils
 				return true;
 			});
 						
-			if (ContextLevelInstance && !ContextLevelInstance->IsEditing())
+			if (ContextLevelInstance && IsExperimentalSettingEnabled(ContextLevelInstance) && !ContextLevelInstance->IsEditing())
 			{
 				FToolMenuSection& Section = CreateLevelSection(Menu);
 				;
@@ -516,12 +536,12 @@ namespace LevelInstanceMenuUtils
 		
 		virtual bool IsClassAllowed(const FClassViewerInitializationOptions& InInitOptions, const UClass* InClass, TSharedRef< FClassViewerFilterFuncs > InFilterFuncs) override
 		{
-			return InClass && InClass->IsChildOf(ALevelInstance::StaticClass()) && !InClass->HasAnyClassFlags(CLASS_Deprecated);
+			return InClass && InClass->IsChildOf(ALevelInstance::StaticClass()) && !InClass->IsChildOf(APackedLevelActor::StaticClass()) && !InClass->HasAnyClassFlags(CLASS_Deprecated);
 		}
 
 		virtual bool IsUnloadedClassAllowed(const FClassViewerInitializationOptions& InInitOptions, const TSharedRef< const IUnloadedBlueprintData > InUnloadedClassData, TSharedRef< FClassViewerFilterFuncs > InFilterFuncs) override
 		{
-			return InUnloadedClassData->IsChildOf(ALevelInstance::StaticClass()) && !InUnloadedClassData->HasAnyClassFlags(CLASS_Deprecated);
+			return InUnloadedClassData->IsChildOf(ALevelInstance::StaticClass()) && !InUnloadedClassData->IsChildOf(APackedLevelActor::StaticClass())  && !InUnloadedClassData->HasAnyClassFlags(CLASS_Deprecated);
 		}
 	};
 
@@ -697,6 +717,11 @@ void FLevelInstanceEditorModule::ExtendContextMenu()
 					return LevelInstanceSubsystem->CanPackAllLoadedActors();
 				}
 				return false;
+			}),
+			FIsActionChecked(),
+			FIsActionButtonVisible::CreateLambda([]() 
+			{
+				return GetDefault<UEditorExperimentalSettings>()->bPackedLevelActor;
 			}));
 
 		FToolMenuEntry& Entry = Section.AddMenuEntry(NAME_None, LOCTEXT("PackLevelActorsTitle", "Pack Level Actors"),
@@ -745,6 +770,11 @@ void FLevelInstanceEditorModule::ExtendContextMenu()
 	{
 		FToolMenuSection& Section = WorldAssetMenu->AddDynamicSection("ActorLevelInstance", FNewToolMenuDelegate::CreateLambda([this](UToolMenu* ToolMenu)
 		{
+			if(!GetDefault<UEditorExperimentalSettings>()->bLevelInstance)
+			{
+				return;
+			}
+
 			if (ToolMenu)
 			{
 				if (UContentBrowserAssetContextMenuContext* AssetMenuContext = ToolMenu->Context.FindContext<UContentBrowserAssetContextMenuContext>())
diff --git a/Engine/Source/Editor/UnrealEd/Classes/Settings/EditorExperimentalSettings.h b/Engine/Source/Editor/UnrealEd/Classes/Settings/EditorExperimentalSettings.h
--- a/Engine/Source/Editor/UnrealEd/Classes/Settings/EditorExperimentalSettings.h
+++ b/Engine/Source/Editor/UnrealEd/Classes/Settings/EditorExperimentalSettings.h
@@ -149,6 +149,14 @@ public:
 	UPROPERTY(EditAnywhere, config, Category = "Content Browser", meta = (DisplayName = "Enable support for long paths (> 260 characters)"))
 	bool bEnableLongPathsSupport;
 
+	/** Allows creating APackedLevelActor blueprint actors */
+	UPROPERTY(EditAnywhere, config, Category = Level)
+	bool bPackedLevelActor;
+
+	/** Allows creating ALevelInstance actors */
+	UPROPERTY(EditAnywhere, config, Category = Level)
+	bool bLevelInstance;
+
 	/**
 	 * Returns an event delegate that is executed when a setting has changed.
 	 *
diff --git a/Engine/Source/Editor/UnrealEd/Private/Settings/SettingsClasses.cpp b/Engine/Source/Editor/UnrealEd/Private/Settings/SettingsClasses.cpp
--- a/Engine/Source/Editor/UnrealEd/Private/Settings/SettingsClasses.cpp
+++ b/Engine/Source/Editor/UnrealEd/Private/Settings/SettingsClasses.cpp
@@ -157,6 +157,8 @@ UEditorExperimentalSettings::UEditorExperimentalSettings( const FObjectInitializ
 	, HDREditorNITLevel(160.0f)
 	, bUseOpenCLForConvexHullDecomp(false)
 	, bAllowPotentiallyUnsafePropertyEditing(false)
+	, bPackedLevelActor(true)
+	, bLevelInstance(true)
 {
 }
 
