Level Instance: remove 'Save Level As'

diff --git a/Engine/Source/Editor/LevelInstanceEditor/Private/LevelInstanceEditorModule.cpp b/Engine/Source/Editor/LevelInstanceEditor/Private/LevelInstanceEditorModule.cpp
--- a/Engine/Source/Editor/LevelInstanceEditor/Private/LevelInstanceEditorModule.cpp
+++ b/Engine/Source/Editor/LevelInstanceEditor/Private/LevelInstanceEditorModule.cpp
@@ -353,50 +353,6 @@ namespace LevelInstanceMenuUtils
 		}
 	}
 		
-	void CreateSaveAsMenu(UToolMenu* Menu, AActor* ContextActor)
-	{
-		ALevelInstance* ContextLevelInstance = nullptr;
-		if (ContextActor)
-		{
-			if (ULevelInstanceSubsystem* LevelInstanceSubsystem = ContextActor->GetWorld()->GetSubsystem<ULevelInstanceSubsystem>())
-			{
-				LevelInstanceSubsystem->ForEachLevelInstanceAncestorsAndSelf(ContextActor, [&ContextLevelInstance](ALevelInstance* LevelInstanceActor)
-					{
-						if (LevelInstanceActor->IsEditing())
-						{
-							ContextLevelInstance = LevelInstanceActor;
-							return false;
-						}
-						return true;
-					});
-			}
-		}
-
-		if (ContextLevelInstance && IsExperimentalSettingEnabled(ContextLevelInstance))
-		{
-			FToolMenuSection& Section = CreateLevelSection(Menu);
-			FText EntryDesc = LOCTEXT("LevelInstanceEditSubMenuEntry", "");
-			const bool bCanCommit = ContextLevelInstance->CanCommit(&EntryDesc);
-
-			FToolUIAction SaveAction;
-			SaveAction.ExecuteAction.BindLambda([ContextLevelInstance](const FToolMenuContext& MenuContext)
-				{
-					ContextLevelInstance->SaveAs();
-				});
-			SaveAction.CanExecuteAction.BindLambda([bCanCommit](const FToolMenuContext& MenuContext)
-				{
-					return bCanCommit;
-				});
-
-			Section.AddMenuEntry(
-				"SaveLevelInstanceAs",
-				LOCTEXT("SaveLevelInstanceAs", "Save Level as..."),
-				TAttribute<FText>(),
-				TAttribute<FSlateIcon>(),
-				SaveAction);
-		}
-	}
-
 	void CreateBreakSubMenu(UToolMenu* Menu, ALevelInstance* ContextLevelInstance)
 	{
 		static int32 BreakLevels = 1;
@@ -749,7 +705,6 @@ void FLevelInstanceEditorModule::ExtendContextMenu()
 			{
 				LevelInstanceMenuUtils::CreateEditMenu(ToolMenu, ContextActor);
 				LevelInstanceMenuUtils::CreateCommitDiscardMenu(ToolMenu, ContextActor);
-				LevelInstanceMenuUtils::CreateSaveAsMenu(ToolMenu, ContextActor);
 				LevelInstanceMenuUtils::CreateBreakMenu(ToolMenu, ContextActor);
 				LevelInstanceMenuUtils::CreatePackedBlueprintMenu(ToolMenu, ContextActor);
 			}
diff --git a/Engine/Source/Runtime/Engine/Private/LevelInstance/LevelInstanceActor.cpp b/Engine/Source/Runtime/Engine/Private/LevelInstance/LevelInstanceActor.cpp
--- a/Engine/Source/Runtime/Engine/Private/LevelInstance/LevelInstanceActor.cpp
+++ b/Engine/Source/Runtime/Engine/Private/LevelInstance/LevelInstanceActor.cpp
@@ -787,13 +787,6 @@ void ALevelInstance::Discard()
 	LevelInstanceSubsystem->CommitLevelInstance(this, bDiscardEdits);
 }
 
-void ALevelInstance::SaveAs()
-{
-	ULevelInstanceSubsystem* LevelInstanceSubsystem = GetLevelInstanceSubsystem();
-	check(LevelInstanceSubsystem);
-	LevelInstanceSubsystem->SaveLevelInstanceAs(this);
-}
-
 bool ALevelInstance::HasDirtyChildren() const
 {
 	if (ULevelInstanceSubsystem* LevelInstanceSubsystem = GetLevelInstanceSubsystem())
diff --git a/Engine/Source/Runtime/Engine/Private/LevelInstance/LevelInstanceSubsystem.cpp b/Engine/Source/Runtime/Engine/Private/LevelInstance/LevelInstanceSubsystem.cpp
--- a/Engine/Source/Runtime/Engine/Private/LevelInstance/LevelInstanceSubsystem.cpp
+++ b/Engine/Source/Runtime/Engine/Private/LevelInstance/LevelInstanceSubsystem.cpp
@@ -1850,37 +1850,6 @@ bool ULevelInstanceSubsystem::CommitLevelInstanceInternal(TUniquePtr<FLevelInsta
 	return true;
 }
 
-void ULevelInstanceSubsystem::SaveLevelInstanceAs(ALevelInstance* LevelInstanceActor)
-{
-	check(CanCommitLevelInstance(LevelInstanceActor));
-
-	const FLevelInstanceEdit* OldLevelInstanceEdit = GetLevelInstanceEdit(LevelInstanceActor);
-	check(OldLevelInstanceEdit);
-	UWorld* EditingWorld = OldLevelInstanceEdit->GetEditWorld();
-	check(EditingWorld);
-
-	// Reset the level transform before saving
-	OldLevelInstanceEdit->LevelStreaming->GetLoadedLevel()->ApplyWorldOffset(-LevelInstanceActor->GetTransform().GetLocation(), false);
-
-	TArray<UObject*> OutObjects;
-	FEditorFileUtils::SaveAssetsAs({ EditingWorld }, OutObjects);
-
-	if (OutObjects.Num() == 0 || OutObjects[0] == EditingWorld)
-	{
-		UE_LOG(LogLevelInstance, Warning, TEXT("Failed to save Level as new asset"));
-		return;
-	}
-
-	UWorld* SavedWorld = StaticCast<UWorld*>(OutObjects[0]);
-	// Discard edits and unload streaming level
-	ResetEdit(LevelInstanceEdit);
-	
-	LevelInstanceActor->SetWorldAsset(SavedWorld);
-
-	LoadLevelInstance(LevelInstanceActor);
-	GEditor->SelectActor(LevelInstanceActor, true, true);
-}
-
 ALevelInstance* ULevelInstanceSubsystem::GetParentLevelInstance(const AActor* Actor) const
 {
 	check(Actor);
diff --git a/Engine/Source/Runtime/Engine/Public/LevelInstance/LevelInstanceActor.h b/Engine/Source/Runtime/Engine/Public/LevelInstance/LevelInstanceActor.h
--- a/Engine/Source/Runtime/Engine/Public/LevelInstance/LevelInstanceActor.h
+++ b/Engine/Source/Runtime/Engine/Public/LevelInstance/LevelInstanceActor.h
@@ -95,7 +95,6 @@ public:
 	void Edit(AActor* ContextActor = nullptr);
 	void Commit();
 	void Discard();
-	void SaveAs();
 	bool HasDirtyChildren() const;
 	bool IsDirty() const;
 	bool SetCurrent();
diff --git a/Engine/Source/Runtime/Engine/Public/LevelInstance/LevelInstanceSubsystem.h b/Engine/Source/Runtime/Engine/Public/LevelInstance/LevelInstanceSubsystem.h
--- a/Engine/Source/Runtime/Engine/Public/LevelInstance/LevelInstanceSubsystem.h
+++ b/Engine/Source/Runtime/Engine/Public/LevelInstance/LevelInstanceSubsystem.h
@@ -61,7 +61,6 @@ public:
 	bool CanDiscardLevelInstance(const ALevelInstance* LevelInstanceActor, FText* OutReason = nullptr) const;
 	void EditLevelInstance(ALevelInstance* LevelInstanceActor, TWeakObjectPtr<AActor> ContextActorPtr = nullptr);
 	bool CommitLevelInstance(ALevelInstance* LevelInstanceActor, bool bDiscardEdits = false, TSet<FName>* DirtyPackages = nullptr);
-	void SaveLevelInstanceAs(ALevelInstance* LevelInstanceActor);
 	bool IsEditingLevelInstanceDirty(const ALevelInstance* LevelInstanceActor) const;
 	bool IsEditingLevelInstance(const ALevelInstance* LevelInstanceActor) const { return GetLevelInstanceEdit(LevelInstanceActor) != nullptr; }
 	
