Fixed a crash caused by not having a ContextActor while generating the submenu Actor -> Level.

The old actor global selection doesn't always contains actors only. It can also contains instances. This is a oddity of the current implementation that should disappear when we will deprecate the global selection sets and the USelection class.

diff --git a/Engine/Source/Editor/LevelInstanceEditor/Private/LevelInstanceEditorModule.cpp b/Engine/Source/Editor/LevelInstanceEditor/Private/LevelInstanceEditorModule.cpp
--- a/Engine/Source/Editor/LevelInstanceEditor/Private/LevelInstanceEditorModule.cpp
+++ b/Engine/Source/Editor/LevelInstanceEditor/Private/LevelInstanceEditorModule.cpp
@@ -196,15 +196,7 @@ namespace LevelInstanceMenuUtils
 		ALevelInstance* LevelInstanceEdit = nullptr;
 		if (ULevelInstanceSubsystem* LevelInstanceSubsystem = ContextActor->GetWorld()->GetSubsystem<ULevelInstanceSubsystem>())
 		{
-			if (ContextActor)
-			{
-				LevelInstanceEdit = LevelInstanceSubsystem->GetEditingLevelInstance();
-			}
-
-			if (!LevelInstanceEdit)
-			{
-				LevelInstanceEdit = LevelInstanceSubsystem->GetEditingLevelInstance();
-			}
+			LevelInstanceEdit = LevelInstanceSubsystem->GetEditingLevelInstance();
 
 			if (LevelInstanceEdit)
 			{
@@ -707,9 +699,9 @@ void FLevelInstanceEditorModule::ExtendContextMenu()
 				LevelInstanceMenuUtils::CreateCommitDiscardMenu(ToolMenu, ContextActor);
 				LevelInstanceMenuUtils::CreateBreakMenu(ToolMenu, ContextActor);
 				LevelInstanceMenuUtils::CreatePackedBlueprintMenu(ToolMenu, ContextActor);
+				LevelInstanceMenuUtils::CreateSetCurrentMenu(ToolMenu, ContextActor);
 			}
 
-			LevelInstanceMenuUtils::CreateSetCurrentMenu(ToolMenu, ContextActor);
 			LevelInstanceMenuUtils::CreateMoveSelectionToMenu(ToolMenu);
 		}
 
