Make Level Instance Editor mode incompatible with Foliage and Landscape Modes.

diff --git a/Engine/Source/Editor/LevelInstanceEditor/Private/LevelInstanceEditorMode.cpp b/Engine/Source/Editor/LevelInstanceEditor/Private/LevelInstanceEditorMode.cpp
--- a/Engine/Source/Editor/LevelInstanceEditor/Private/LevelInstanceEditorMode.cpp
+++ b/Engine/Source/Editor/LevelInstanceEditor/Private/LevelInstanceEditorMode.cpp
@@ -4,6 +4,7 @@
 #include "LevelInstanceEditorModeToolkit.h"
 #include "LevelInstanceEditorModeCommands.h"
 #include "Editor.h"
+#include "EditorModes.h"
 #include "Engine/World.h"
 #include "LevelInstance/LevelInstanceSubsystem.h"
 #include "LevelInstance/LevelInstanceActor.h"
@@ -77,6 +78,11 @@ void ULevelInstanceEditorMode::CreateToolkit()
 	Toolkit = MakeShared<FLevelInstanceEditorModeToolkit>();
 }
 
+bool ULevelInstanceEditorMode::IsCompatibleWith(FEditorModeID OtherModeID) const
+{
+	return (OtherModeID != FBuiltinEditorModes::EM_Foliage) && (OtherModeID != FBuiltinEditorModes::EM_Landscape);
+}
+
 void ULevelInstanceEditorMode::BindCommands()
 {
 	UEdMode::BindCommands();
diff --git a/Engine/Source/Editor/LevelInstanceEditor/Private/LevelInstanceEditorMode.h b/Engine/Source/Editor/LevelInstanceEditor/Private/LevelInstanceEditorMode.h
--- a/Engine/Source/Editor/LevelInstanceEditor/Private/LevelInstanceEditorMode.h
+++ b/Engine/Source/Editor/LevelInstanceEditor/Private/LevelInstanceEditorMode.h
@@ -24,7 +24,7 @@ public:
 	virtual void Exit() override;
 	virtual void CreateToolkit() override;
 	
-	virtual bool IsCompatibleWith(FEditorModeID OtherModeID) const override { return true; }
+	virtual bool IsCompatibleWith(FEditorModeID OtherModeID) const override;
 
 	/** Check to see if an actor can be selected in this mode - no side effects */
 	virtual bool IsSelectionDisallowed(AActor* InActor, bool bInSelection) const override;
