LevelInstance: Allow overriding LevelInstance class in Editor config

[/Script/LevelInstanceEditor.LevelInstanceEditorSettings]
LevelInstanceClassName=Module.ClassName

diff --git a/Engine/Source/Editor/LevelInstanceEditor/Private/LevelInstanceEditorModule.cpp b/Engine/Source/Editor/LevelInstanceEditor/Private/LevelInstanceEditorModule.cpp
--- a/Engine/Source/Editor/LevelInstanceEditor/Private/LevelInstanceEditorModule.cpp
+++ b/Engine/Source/Editor/LevelInstanceEditor/Private/LevelInstanceEditorModule.cpp
@@ -301,11 +301,21 @@ namespace LevelInstanceMenuUtils
 			FString TemplateMapPackage;
 			bool bOutIsPartitionedWorld = false;
 			const bool bShowPartitionedTemplates = false;
-			if (!GetMutableDefault<ULevelInstanceEditorSettings>()->TemplateMapInfos.Num() || NewLevelDialogModule.CreateAndShowTemplateDialog(MainFrameModule.GetParentWindow(), LOCTEXT("LevelInstanceTemplateDialog", "Choose Level Instance Template..."), GetMutableDefault<ULevelInstanceEditorSettings>()->TemplateMapInfos, TemplateMapPackage, bShowPartitionedTemplates, bOutIsPartitionedWorld))
+			ULevelInstanceEditorSettings* LevelInstanceEditorSettings = GetMutableDefault<ULevelInstanceEditorSettings>();
+			if (!LevelInstanceEditorSettings->TemplateMapInfos.Num() || NewLevelDialogModule.CreateAndShowTemplateDialog(MainFrameModule.GetParentWindow(), LOCTEXT("LevelInstanceTemplateDialog", "Choose Level Instance Template..."), GetMutableDefault<ULevelInstanceEditorSettings>()->TemplateMapInfos, TemplateMapPackage, bShowPartitionedTemplates, bOutIsPartitionedWorld))
 			{
 				UPackage* TemplatePackage = !TemplateMapPackage.IsEmpty() ? LoadPackage(nullptr, *TemplateMapPackage, LOAD_None) : nullptr;
 				
 				CreationParams.TemplateWorld = TemplatePackage ? UWorld::FindWorldInPackage(TemplatePackage) : nullptr;
+				
+				if (!LevelInstanceEditorSettings->LevelInstanceClassName.IsEmpty())
+				{
+					UClass* LevelInstanceClass = LoadClass<AActor>(nullptr, *LevelInstanceEditorSettings->LevelInstanceClassName, nullptr, LOAD_NoWarn);
+					if (LevelInstanceClass && LevelInstanceClass->ImplementsInterface(ULevelInstanceInterface::StaticClass()))
+					{
+						CreationParams.LevelInstanceClass = LevelInstanceClass;
+					}
+				}
 
 				if (!LevelInstanceSubsystem->CreateLevelInstanceFrom(ActorsToMove, CreationParams))
 				{
diff --git a/Engine/Source/Editor/LevelInstanceEditor/Public/LevelInstanceEditorSettings.h b/Engine/Source/Editor/LevelInstanceEditor/Public/LevelInstanceEditorSettings.h
--- a/Engine/Source/Editor/LevelInstanceEditor/Public/LevelInstanceEditorSettings.h
+++ b/Engine/Source/Editor/LevelInstanceEditor/Public/LevelInstanceEditorSettings.h
@@ -19,4 +19,7 @@ public:
 	/** List of info for all known LevelInstance template maps */
 	UPROPERTY(config)
 	TArray<FTemplateMapInfo> TemplateMapInfos;
+
+	UPROPERTY(config)
+	FString LevelInstanceClassName;
 };
\ No newline at end of file
