Added option for a Level referenced by a Level Instance to add Data Layer support.

diff --git a/Engine/Source/Editor/LevelInstanceEditor/Private/LevelInstanceEditorModule.cpp b/Engine/Source/Editor/LevelInstanceEditor/Private/LevelInstanceEditorModule.cpp
--- a/Engine/Source/Editor/LevelInstanceEditor/Private/LevelInstanceEditorModule.cpp
+++ b/Engine/Source/Editor/LevelInstanceEditor/Private/LevelInstanceEditorModule.cpp
@@ -151,6 +151,32 @@ namespace LevelInstanceMenuUtils
 			}
 		}
 	}
+
+	void CreateAddDataLayerSupportMenu(UToolMenu* Menu, AActor* ContextActor)
+	{
+		FText EntryDesc;
+		ILevelInstanceInterface* LevelInstance = Cast<ILevelInstanceInterface>(ContextActor);
+		ULevelInstanceSubsystem* LevelInstanceSubsystem = ContextActor->GetWorld()->GetSubsystem<ULevelInstanceSubsystem>();
+		if (!LevelInstance || 
+			!LevelInstanceSubsystem ||
+			!LevelInstance->IsWorldAssetValid() ||
+			!LevelInstance->CanEnterEdit(&EntryDesc) ||
+			!ULevelInstanceSubsystem::CanUsePackage(*LevelInstance->GetWorldAssetPackage()) ||
+			ULevel::GetPartitionedLevelCanBeUsedByLevelInstanceFromPackage(*LevelInstance->GetWorldAssetPackage()))
+		{
+			return;
+		}
+
+		FToolUIAction LevelInstanceEditAction;
+		LevelInstanceEditAction.ExecuteAction.BindLambda([LevelInstanceSubsystem, LevelInstance, ContextActor](const FToolMenuContext&)
+		{
+			LevelInstanceSubsystem->AddDataLayerSupport(LevelInstance);
+		});
+
+		AActor* LevelInstanceActor = CastChecked<AActor>(LevelInstance);
+		EntryDesc = FText::Format(LOCTEXT("LevelInstanceName", "{0}:{1}"), FText::FromString(LevelInstanceActor->GetActorLabel()), FText::FromString(LevelInstance->GetWorldAssetPackage()));
+		CreateLevelSection(Menu).AddMenuEntry(NAME_None, LOCTEXT("AddDataLayerSupport", "Add Data Layer Support"), EntryDesc, FSlateIcon(), LevelInstanceEditAction);
+	}
 		
 	void CreateCommitDiscardMenu(UToolMenu* Menu, AActor* ContextActor)
 	{
@@ -709,6 +735,7 @@ void FLevelInstanceEditorModule::ExtendContextMenu()
 			if (ContextActor)
 			{
 				LevelInstanceMenuUtils::CreateEditMenu(ToolMenu, ContextActor);
+				LevelInstanceMenuUtils::CreateAddDataLayerSupportMenu(ToolMenu, ContextActor);
 				LevelInstanceMenuUtils::CreateCommitDiscardMenu(ToolMenu, ContextActor);
 				LevelInstanceMenuUtils::CreateBreakMenu(ToolMenu, ContextActor);
 				LevelInstanceMenuUtils::CreatePackedBlueprintMenu(ToolMenu, ContextActor);
diff --git a/Engine/Source/Editor/WorldPartitionEditor/Private/WorldPartition/Customizations/WorldPartitionDetailsCustomization.cpp b/Engine/Source/Editor/WorldPartitionEditor/Private/WorldPartition/Customizations/WorldPartitionDetailsCustomization.cpp
--- a/Engine/Source/Editor/WorldPartitionEditor/Private/WorldPartition/Customizations/WorldPartitionDetailsCustomization.cpp
+++ b/Engine/Source/Editor/WorldPartitionEditor/Private/WorldPartition/Customizations/WorldPartitionDetailsCustomization.cpp
@@ -68,6 +68,14 @@ void FWorldPartitionDetails::HandleWorldPartitionEnableStreamingChanged(ECheckBo
 {
 	if (CheckState == ECheckBoxState::Checked)
 	{
+		if (WorldPartition->CanBeUsedByLevelInstance())
+		{
+			if (FMessageDialog::Open(EAppMsgType::YesNo, LOCTEXT("WorldPartitionConfirmEnableStreamingOnLevelReferencedByLevelInstance", "You are about to enable streaming on a level that could be used by level instances. Continuing will invalidate level instances referencing this level. Continue?")) == EAppReturnType::No)
+			{
+				return;
+			}
+		}
+
 		if (!WorldPartition->bStreamingWasEnabled)
 		{
 			if (FMessageDialog::Open(EAppMsgType::YesNo, LOCTEXT("WorldPartitionConfirmEnableStreaming", "You are about to enable streaming for the first time, the world will be setup to stream. Continue?")) == EAppReturnType::No)
diff --git a/Engine/Source/Runtime/Engine/Classes/Engine/Level.h b/Engine/Source/Runtime/Engine/Classes/Engine/Level.h
--- a/Engine/Source/Runtime/Engine/Classes/Engine/Level.h
+++ b/Engine/Source/Runtime/Engine/Classes/Engine/Level.h
@@ -78,6 +78,7 @@ private:
 	friend class UWorld;
 	friend class ULevel;
 	friend class UActorFolder;
+	friend class FWorldPartitionConverter;
 	friend class UWorldPartitionConvertCommandlet;
 	friend class FWorldPartitionLevelHelper;
 	friend class UWorldPartitionLevelStreamingDynamic;
diff --git a/Engine/Source/Runtime/Engine/Private/LevelInstance/LevelInstanceSubsystem.cpp b/Engine/Source/Runtime/Engine/Private/LevelInstance/LevelInstanceSubsystem.cpp
--- a/Engine/Source/Runtime/Engine/Private/LevelInstance/LevelInstanceSubsystem.cpp
+++ b/Engine/Source/Runtime/Engine/Private/LevelInstance/LevelInstanceSubsystem.cpp
@@ -23,6 +23,7 @@
 #include "LevelInstance/LevelInstanceEditorPivotActor.h"
 #include "WorldPartition/DataLayer/DataLayerSubsystem.h"
 #include "WorldPartition/DataLayer/DataLayerInstanceWithAsset.h"
+#include "WorldPartition/WorldPartitionConverter.h"
 #include "WorldPartition/WorldPartitionMiniMap.h"
 #include "Misc/ScopedSlowTask.h"
 #include "Misc/ITransaction.h"
@@ -1925,6 +1926,39 @@ bool ULevelInstanceSubsystem::CommitLevelInstance(ILevelInstanceInterface* Level
 	return false;
 }
 
+bool ULevelInstanceSubsystem::AddDataLayerSupport(ILevelInstanceInterface* LevelInstance)
+{
+	check(LevelInstance);
+	
+	bool bSuccess = false;
+	EditLevelInstance(LevelInstance);
+	
+	if (GetEditingLevelInstance())
+	{
+		UWorld* World = Cast<UWorld>(LevelInstance->GetWorldAsset().Get());
+		if (World && !World->IsPartitionedWorld())
+		{
+			FWorldPartitionConverter::FParameters Parameters;
+			Parameters.bConvertSubLevels = false;
+			Parameters.bEnableStreaming = false;
+			Parameters.bUseActorFolders = true;
+			if (FWorldPartitionConverter::Convert(World, Parameters))
+			{
+				UWorldPartition* WorldPartition = World->GetWorldPartition();
+				check(WorldPartition);
+				check(WorldPartition->CanBeUsedByLevelInstance());
+				check(World->PersistentLevel->IsUsingActorFolders());
+				bSuccess = true;
+			}
+		}
+
+		CommitLevelInstance(LevelInstance, !bSuccess);
+	}
+
+	UE_CLOG(!bSuccess, LogLevelInstance, Error, TEXT("Failed to convert level %s for Level Instance Data Layer support."), *LevelInstance->GetWorldAssetPackage());
+	return bSuccess;
+}
+
 bool ULevelInstanceSubsystem::CommitLevelInstanceInternal(TUniquePtr<FLevelInstanceEdit>& InLevelInstanceEdit, bool bDiscardEdits, bool bDiscardOnFailure, TSet<FName>* DirtyPackages)
 {
 	TGuardValue<bool> CommitScope(bIsCommittingLevelInstance, true);
diff --git a/Engine/Source/Runtime/Engine/Private/WorldPartition/WorldPartitionConverter.cpp b/Engine/Source/Runtime/Engine/Private/WorldPartition/WorldPartitionConverter.cpp
--- a/Engine/Source/Runtime/Engine/Private/WorldPartition/WorldPartitionConverter.cpp
+++ b/Engine/Source/Runtime/Engine/Private/WorldPartition/WorldPartitionConverter.cpp
@@ -57,8 +57,10 @@ bool FWorldPartitionConverter::Convert()
 		{
 			return false;
 		}
+		FLevelActorFoldersHelper::SetUseActorFolders(World->PersistentLevel, Parameters.bUseActorFolders);
 		WorldPartition->bEnableStreaming = Parameters.bEnableStreaming;
 		WorldPartition->bStreamingWasEnabled = Parameters.bEnableStreaming;
+		WorldPartition->SetCanBeUsedByLevelInstance(!Parameters.bEnableStreaming);
 		bCreatedWorldPartition = true;
 	}
 
diff --git a/Engine/Source/Runtime/Engine/Public/LevelInstance/LevelInstanceSubsystem.h b/Engine/Source/Runtime/Engine/Public/LevelInstance/LevelInstanceSubsystem.h
--- a/Engine/Source/Runtime/Engine/Public/LevelInstance/LevelInstanceSubsystem.h
+++ b/Engine/Source/Runtime/Engine/Public/LevelInstance/LevelInstanceSubsystem.h
@@ -66,6 +66,7 @@ public:
 	bool CanCommitLevelInstance(const ILevelInstanceInterface* LevelInstance, bool bDiscardEdits = false, FText* OutReason = nullptr) const;
 	void EditLevelInstance(ILevelInstanceInterface* LevelInstance, TWeakObjectPtr<AActor> ContextActorPtr = nullptr);
 	bool CommitLevelInstance(ILevelInstanceInterface* LevelInstance, bool bDiscardEdits = false, TSet<FName>* DirtyPackages = nullptr);
+	bool AddDataLayerSupport(ILevelInstanceInterface* LevelInstance);
 	bool IsEditingLevelInstanceDirty(const ILevelInstanceInterface* LevelInstance) const;
 	bool IsEditingLevelInstance(const ILevelInstanceInterface* LevelInstance) const { return GetLevelInstanceEdit(LevelInstance) != nullptr; }
 	
diff --git a/Engine/Source/Runtime/Engine/Public/WorldPartition/WorldPartitionConverter.h b/Engine/Source/Runtime/Engine/Public/WorldPartition/WorldPartitionConverter.h
--- a/Engine/Source/Runtime/Engine/Public/WorldPartition/WorldPartitionConverter.h
+++ b/Engine/Source/Runtime/Engine/Public/WorldPartition/WorldPartitionConverter.h
@@ -21,12 +21,14 @@ public:
 		FParameters()
 		: bConvertSubLevels(true)
 		, bEnableStreaming(true)
+		, bUseActorFolders(true)
 		, EditorHashClass(nullptr)
 		, RuntimeHashClass(nullptr)
 		{}
 
 		bool bConvertSubLevels;
 		bool bEnableStreaming;
+		bool bUseActorFolders;
 		TSubclassOf<UWorldPartitionEditorHash> EditorHashClass;
 		TSubclassOf<UWorldPartitionRuntimeHash> RuntimeHashClass;
 	};
