LevelInstance: Exit Mode shortcut checks that it doesn't interfere with SelectNone shortcut

diff --git a/Engine/Source/Editor/LevelInstanceEditor/Private/LevelInstanceEditorMode.cpp b/Engine/Source/Editor/LevelInstanceEditor/Private/LevelInstanceEditorMode.cpp
--- a/Engine/Source/Editor/LevelInstanceEditor/Private/LevelInstanceEditorMode.cpp
+++ b/Engine/Source/Editor/LevelInstanceEditor/Private/LevelInstanceEditorMode.cpp
@@ -4,6 +4,7 @@
 #include "LevelInstanceEditorModeToolkit.h"
 #include "LevelInstanceEditorModeCommands.h"
 #include "Editor.h"
+#include "Selection.h"
 #include "EditorModes.h"
 #include "Engine/World.h"
 #include "LevelInstance/LevelInstanceSubsystem.h"
@@ -91,7 +92,27 @@ void ULevelInstanceEditorMode::BindCommands()
 
 	CommandList->MapAction(
 		Commands.ExitMode,
-		FExecuteAction::CreateUObject(this, &ULevelInstanceEditorMode::ExitModeCommand));
+		FExecuteAction::CreateUObject(this, &ULevelInstanceEditorMode::ExitModeCommand),
+		FCanExecuteAction::CreateLambda([&] 
+		{ 
+			// If some actors are selected make sure we don't interfere with the SelectNone command
+			if(GEditor->GetSelectedActors()->Num() > 0)
+			{
+				const FInputChord& SelectNonePrimary = FLevelEditorCommands::Get().SelectNone->GetActiveChord(EMultipleKeyBindingIndex::Primary).Get();
+				if (SelectNonePrimary.IsValidChord() && Commands.ExitMode->HasActiveChord(SelectNonePrimary))
+				{
+					return false;
+				}
+
+				const FInputChord& SelectNoneSecondary = FLevelEditorCommands::Get().SelectNone->GetActiveChord(EMultipleKeyBindingIndex::Secondary).Get();
+				if (SelectNoneSecondary.IsValidChord() && Commands.ExitMode->HasActiveChord(SelectNoneSecondary))
+				{
+					return false;
+				}
+			}
+
+			return true;
+		}));
 
 	CommandList->MapAction(
 		Commands.ToggleContextRestriction,
