Level Instance: Persist create settings & add new setting make the create dialog optional.

diff --git a/Engine/Source/Editor/LevelInstanceEditor/LevelInstanceEditor.Build.cs b/Engine/Source/Editor/LevelInstanceEditor/LevelInstanceEditor.Build.cs
--- a/Engine/Source/Editor/LevelInstanceEditor/LevelInstanceEditor.Build.cs
+++ b/Engine/Source/Editor/LevelInstanceEditor/LevelInstanceEditor.Build.cs
@@ -32,7 +32,8 @@ public class LevelInstanceEditor : ModuleRules
 				"AssetTools",
 				"ClassViewer",
 				"MessageLog",
-				"EditorWidgets"
+				"EditorWidgets",
+				"DeveloperSettings"
 			}
 		);
 
diff --git a/Engine/Source/Editor/LevelInstanceEditor/Private/LevelInstanceEditorModule.cpp b/Engine/Source/Editor/LevelInstanceEditor/Private/LevelInstanceEditorModule.cpp
--- a/Engine/Source/Editor/LevelInstanceEditor/Private/LevelInstanceEditorModule.cpp
+++ b/Engine/Source/Editor/LevelInstanceEditor/Private/LevelInstanceEditorModule.cpp
@@ -312,16 +312,22 @@ namespace LevelInstanceMenuUtils
 		const bool bForceExternalActors = LevelInstanceSubsystem->GetWorld()->IsPartitionedWorld();
 		FNewLevelInstanceParams& DialogParams = NewLevelInstanceDialog->GetCreationParams();
 		DialogParams.Type = CreationType;
+		DialogParams.bAlwaysShowDialog = GetDefault<ULevelInstanceEditorPerProjectUserSettings>()->bAlwaysShowDialog;
+		DialogParams.PivotType = GetDefault<ULevelInstanceEditorPerProjectUserSettings>()->PivotType;
+		DialogParams.PivotActor = DialogParams.PivotType == ELevelInstancePivotType::Actor ? ActorsToMove[0] : nullptr;
 		DialogParams.HideCreationType();
 		DialogParams.SetForceExternalActors(bForceExternalActors);
 		NewLevelInstanceWindow->SetContent(NewLevelInstanceDialog);
 
-		FSlateApplication::Get().AddModalWindow(NewLevelInstanceWindow.ToSharedRef(), MainFrameModule.GetParentWindow());
-
-
-		if (NewLevelInstanceDialog->ClickedOk())
+		if (GetDefault<ULevelInstanceEditorPerProjectUserSettings>()->bAlwaysShowDialog)
+		{
+			FSlateApplication::Get().AddModalWindow(NewLevelInstanceWindow.ToSharedRef(), MainFrameModule.GetParentWindow());
+		}
+		
+		if (!GetDefault<ULevelInstanceEditorPerProjectUserSettings>()->bAlwaysShowDialog || NewLevelInstanceDialog->ClickedOk())
 		{
 			FNewLevelInstanceParams CreationParams(NewLevelInstanceDialog->GetCreationParams());
+			ULevelInstanceEditorPerProjectUserSettings::UpdateFrom(CreationParams);
 
 			FNewLevelDialogModule& NewLevelDialogModule = FModuleManager::LoadModuleChecked<FNewLevelDialogModule>("NewLevelDialog");
 			FString TemplateMapPackage;
diff --git a/Engine/Source/Editor/LevelInstanceEditor/Private/LevelInstanceEditorSettings.cpp b/Engine/Source/Editor/LevelInstanceEditor/Private/LevelInstanceEditorSettings.cpp
--- a/Engine/Source/Editor/LevelInstanceEditor/Private/LevelInstanceEditorSettings.cpp
+++ b/Engine/Source/Editor/LevelInstanceEditor/Private/LevelInstanceEditorSettings.cpp
@@ -5,3 +5,21 @@
 ULevelInstanceEditorSettings::ULevelInstanceEditorSettings()
 {
 }
+
+ULevelInstanceEditorPerProjectUserSettings::ULevelInstanceEditorPerProjectUserSettings()
+{
+	bAlwaysShowDialog = true;
+	PivotType = ELevelInstancePivotType::CenterMinZ;
+}
+
+void ULevelInstanceEditorPerProjectUserSettings::UpdateFrom(const FNewLevelInstanceParams& Params)
+{
+	ULevelInstanceEditorPerProjectUserSettings* UserSettings = GetMutableDefault<ULevelInstanceEditorPerProjectUserSettings>();
+
+	if (Params.PivotType != UserSettings->PivotType || Params.bAlwaysShowDialog != UserSettings->bAlwaysShowDialog)
+	{
+		UserSettings->bAlwaysShowDialog = Params.bAlwaysShowDialog;
+		UserSettings->PivotType = Params.PivotType;
+		UserSettings->SaveConfig();
+	}
+}
\ No newline at end of file
diff --git a/Engine/Source/Editor/LevelInstanceEditor/Private/SNewLevelInstanceDialog.cpp b/Engine/Source/Editor/LevelInstanceEditor/Private/SNewLevelInstanceDialog.cpp
--- a/Engine/Source/Editor/LevelInstanceEditor/Private/SNewLevelInstanceDialog.cpp
+++ b/Engine/Source/Editor/LevelInstanceEditor/Private/SNewLevelInstanceDialog.cpp
@@ -7,6 +7,7 @@
 #include "Widgets/Layout/SSpacer.h"
 #include "Widgets/Input/SComboBox.h"
 #include "Widgets/Input/SButton.h"
+#include "Widgets/Input/SCheckBox.h"
 #include "PropertyEditorDelegates.h"
 #include "LevelInstance/LevelInstanceSubsystem.h"
 #include "Modules/ModuleManager.h"
@@ -99,6 +100,30 @@ void SNewLevelInstanceDialog::Construct(const FArguments& InArgs)
 					.Text(LOCTEXT("CancelButton", "Cancel"))
 				]
 			]
+			+ SVerticalBox::Slot()
+				.AutoHeight()
+				[
+					SNew(SHorizontalBox)
+					+ SHorizontalBox::Slot()
+				.FillWidth(1.f)
+				[
+					SNew(SCheckBox)
+					.IsChecked_Lambda([this]() -> ECheckBoxState
+					{
+						return CreationParams.bAlwaysShowDialog ? ECheckBoxState::Checked : ECheckBoxState::Unchecked;
+					})
+					.OnCheckStateChanged_Lambda([this](ECheckBoxState State)
+					{
+						CreationParams.bAlwaysShowDialog = State == ECheckBoxState::Checked;
+					})
+					[
+						SNew(STextBlock)
+						.Font(FCoreStyle::Get().GetFontStyle("SmallFont"))
+						.Text(LOCTEXT("AlwaysShowDialog", "Always show dialog"))
+						.ToolTipText(LOCTEXT("AlwaysShowDialogToolTip", "Show this dialog everytime a level instance is created. Can be changed in editor preferences (Content Editors > Level Instance)."))
+					]
+				]
+			]
 		]
 	];
 }
diff --git a/Engine/Source/Editor/LevelInstanceEditor/Public/LevelInstanceEditorSettings.h b/Engine/Source/Editor/LevelInstanceEditor/Public/LevelInstanceEditorSettings.h
--- a/Engine/Source/Editor/LevelInstanceEditor/Public/LevelInstanceEditorSettings.h
+++ b/Engine/Source/Editor/LevelInstanceEditor/Public/LevelInstanceEditorSettings.h
@@ -6,6 +6,8 @@
 #include "UObject/ObjectMacros.h"
 #include "UObject/Object.h"
 #include "Editor/UnrealEdEngine.h"
+#include "Engine/DeveloperSettings.h"
+#include "LevelInstance/LevelInstanceTypes.h"
 #include "LevelInstanceEditorSettings.generated.h"
 
 UCLASS(config = Editor)
@@ -22,4 +24,25 @@ public:
 
 	UPROPERTY(config)
 	FString LevelInstanceClassName;
+};
+
+UCLASS(config = EditorPerProjectUserSettings, meta = (DisplayName = "Level Instance"))
+class ULevelInstanceEditorPerProjectUserSettings : public UDeveloperSettings
+{
+	GENERATED_BODY()
+
+public:
+	ULevelInstanceEditorPerProjectUserSettings();
+
+	static void UpdateFrom(const FNewLevelInstanceParams& Params);
+
+	/** Gets the category for the settings, some high level grouping like, Editor, Engine, Game...etc. */
+	virtual FName GetCategoryName() const override { return TEXT("ContentEditors"); }
+		
+	/** If false, create dialog will not be shown and last settings will be used. */
+	UPROPERTY(config, EditAnywhere, Category = Create)
+	bool bAlwaysShowDialog;
+
+	UPROPERTY(config, EditAnywhere, Category = Pivot)
+	ELevelInstancePivotType PivotType;
 };
\ No newline at end of file
diff --git a/Engine/Source/Runtime/Engine/Public/LevelInstance/LevelInstanceTypes.h b/Engine/Source/Runtime/Engine/Public/LevelInstance/LevelInstanceTypes.h
--- a/Engine/Source/Runtime/Engine/Public/LevelInstance/LevelInstanceTypes.h
+++ b/Engine/Source/Runtime/Engine/Public/LevelInstance/LevelInstanceTypes.h
@@ -80,6 +80,9 @@ struct FNewLevelInstanceParams
 	UPROPERTY(EditAnywhere, Category = Pivot)
 	TObjectPtr<AActor> PivotActor = nullptr;
 
+	UPROPERTY()
+	bool bAlwaysShowDialog = true;
+
 	UPROPERTY()
 	TObjectPtr<UWorld> TemplateWorld = nullptr;
 		
