LevelInstance: Can now drag worlds from content browser to create a level instance

diff --git a/Engine/Source/Editor/LevelInstanceEditor/Private/LevelInstanceActorFactory.cpp b/Engine/Source/Editor/LevelInstanceEditor/Private/LevelInstanceActorFactory.cpp
new file mode 100644
--- /dev/null
+++ b/Engine/Source/Editor/LevelInstanceEditor/Private/LevelInstanceActorFactory.cpp
@@ -0,0 +1,36 @@
+// Copyright Epic Games, Inc. All Rights Reserved.
+
+#include "LevelInstanceActorFactory.h"
+#include "LevelInstanceEditorSettings.h"
+#include "LevelInstance/LevelInstanceInterface.h"
+#include "Settings/EditorExperimentalSettings.h"
+
+ULevelInstanceActorFactory::ULevelInstanceActorFactory(const FObjectInitializer& ObjectInitializer)
+	: Super(ObjectInitializer)
+{
+	NewActorClassName = GetDefault<ULevelInstanceEditorSettings>()->LevelInstanceClassName;
+}
+
+void ULevelInstanceActorFactory::PostSpawnActor(UObject* Asset, AActor* NewActor)
+{
+	ILevelInstanceInterface* LevelInstanceInterface = CastChecked<ILevelInstanceInterface>(NewActor);
+	LevelInstanceInterface->SetWorldAsset(Asset);
+	LevelInstanceInterface->LoadLevelInstance();
+}
+
+bool ULevelInstanceActorFactory::CanCreateActorFrom(const FAssetData& AssetData, FText& OutErrorMsg)
+{
+	if (!GetDefault<UEditorExperimentalSettings>()->bLevelInstance)
+	{
+		OutErrorMsg = NSLOCTEXT("LevelInstanceActorFactory", "ExperimentalSettings", "Level Instance must be enabled in experimental settings.");
+		return false;
+	}
+
+	if (!AssetData.IsValid() || !AssetData.IsInstanceOf(UWorld::StaticClass()))
+	{
+		OutErrorMsg = NSLOCTEXT("LevelInstanceActorFactory", "NoWorld", "A valid world must be specified.");
+		return false;
+	}
+
+	return true;
+}
diff --git a/Engine/Source/Editor/LevelInstanceEditor/Private/LevelInstanceActorFactory.h b/Engine/Source/Editor/LevelInstanceEditor/Private/LevelInstanceActorFactory.h
new file mode 100644
--- /dev/null
+++ b/Engine/Source/Editor/LevelInstanceEditor/Private/LevelInstanceActorFactory.h
@@ -0,0 +1,19 @@
+// Copyright Epic Games, Inc. All Rights Reserved.
+
+#pragma once
+
+#include "CoreMinimal.h"
+#include "UObject/ObjectMacros.h"
+#include "ActorFactories/ActorFactory.h"
+#include "LevelInstanceActorFactory.generated.h"
+
+UCLASS(MinimalAPI, config = Editor)
+class ULevelInstanceActorFactory : public UActorFactory
+{
+	GENERATED_UCLASS_BODY()
+
+	//~ Begin UActorFactory Interface
+	virtual void PostSpawnActor(UObject* Asset, AActor* NewActor) override;
+	virtual bool CanCreateActorFrom(const FAssetData& AssetData, FText& OutErrorMsg) override;
+	//~ End UActorFactory Interface
+};
diff --git a/Engine/Source/Editor/LevelInstanceEditor/Private/LevelInstanceEditorSettings.cpp b/Engine/Source/Editor/LevelInstanceEditor/Private/LevelInstanceEditorSettings.cpp
--- a/Engine/Source/Editor/LevelInstanceEditor/Private/LevelInstanceEditorSettings.cpp
+++ b/Engine/Source/Editor/LevelInstanceEditor/Private/LevelInstanceEditorSettings.cpp
@@ -1,9 +1,11 @@
 // Copyright Epic Games, Inc. All Rights Reserved.
 
 #include "LevelInstanceEditorSettings.h"
+#include "LevelInstance/LevelInstanceActor.h"
 
 ULevelInstanceEditorSettings::ULevelInstanceEditorSettings()
 {
+	LevelInstanceClassName = ALevelInstance::StaticClass()->GetPathName();
 }
 
 ULevelInstanceEditorPerProjectUserSettings::ULevelInstanceEditorPerProjectUserSettings()
diff --git a/Engine/Source/Editor/UnrealEd/Private/LevelEditorDragDropHandler.cpp b/Engine/Source/Editor/UnrealEd/Private/LevelEditorDragDropHandler.cpp
--- a/Engine/Source/Editor/UnrealEd/Private/LevelEditorDragDropHandler.cpp
+++ b/Engine/Source/Editor/UnrealEd/Private/LevelEditorDragDropHandler.cpp
@@ -25,7 +25,7 @@ bool ULevelEditorDragDropHandler::PreviewDropObjectsAtCoordinates(int32 MouseX,
 	bCanDrop = false;
 	HintText = FText::GetEmpty();
 
-	if ( !ObjectTools::IsAssetValidForPlacing(World, AssetData.ObjectPath.ToString() ) )
+	if ( AssetData.ObjectPath.IsNone() )
 	{
 		bCanDrop = false;
 		return false;
diff --git a/Engine/Source/Editor/UnrealEd/Public/ObjectTools.h b/Engine/Source/Editor/UnrealEd/Public/ObjectTools.h
--- a/Engine/Source/Editor/UnrealEd/Public/ObjectTools.h
+++ b/Engine/Source/Editor/UnrealEd/Public/ObjectTools.h
@@ -535,14 +535,7 @@ namespace ObjectTools
 	 */
 	UNREALED_API void RemoveDeletedObjectsFromPropertyWindows( TArray<UObject*>& DeletedObjects );
 
-	/**
-	 * Determines if the asset is placeable in a world.
-	 *
-	 * @param InWorld	The world.
-	 * @param ObjectPath	Object path.
-	 *
-	 * @return true if the asset can be placed in the world.
-	 */
+	UE_DEPRECATED(5.1, "No longer used.")
 	UNREALED_API bool IsAssetValidForPlacing(UWorld* InWorld, const FString& ObjectPath);
 
 	/**
