Level Instance: Allow breaking multiple selection

diff --git a/Engine/Source/Editor/LevelInstanceEditor/Private/LevelInstanceEditorModule.cpp b/Engine/Source/Editor/LevelInstanceEditor/Private/LevelInstanceEditorModule.cpp
--- a/Engine/Source/Editor/LevelInstanceEditor/Private/LevelInstanceEditorModule.cpp
+++ b/Engine/Source/Editor/LevelInstanceEditor/Private/LevelInstanceEditorModule.cpp
@@ -389,13 +389,11 @@ namespace LevelInstanceMenuUtils
 		}
 	}
 		
-	void CreateBreakSubMenu(UToolMenu* Menu, ILevelInstanceInterface* ContextLevelInstance)
+	void CreateBreakSubMenu(UToolMenu* Menu, TArray<ILevelInstanceInterface*> BreakableLevelInstances)
 	{
 		static int32 BreakLevels = 1;
 
-		check(ContextLevelInstance);
-
-		if (ULevelInstanceSubsystem* LevelInstanceSubsystem = ContextLevelInstance->GetLevelInstanceSubsystem())
+		if (ULevelInstanceSubsystem* LevelInstanceSubsystem = GEditor->GetEditorWorldContext().World()->GetSubsystem<ULevelInstanceSubsystem>())
 		{
 			FToolMenuSection& Section = Menu->AddSection(NAME_None, LOCTEXT("LevelInstanceBreakSection", "Break Level Instance"));
 			TSharedRef<SWidget> MenuWidget =
@@ -426,53 +424,56 @@ namespace LevelInstanceMenuUtils
 					SNew(SButton)
 					.HAlign(HAlign_Center)
 					.ContentPadding(FAppStyle::GetMargin("StandardDialog.ContentPadding"))
-					.OnClicked_Lambda([ContextLevelInstance, LevelInstanceSubsystem]() 
+					.OnClicked_Lambda([BreakableLevelInstances, LevelInstanceSubsystem]() 
+					{
+						const FText LevelInstanceBreakWarning = FText::Format(LOCTEXT("BreakingLevelInstance", "You are about to break {0} level instance(s). This action cannot be undone. Are you sure ?"),  FText::FromString(FString::FromInt(BreakableLevelInstances.Num())));
+						if (FMessageDialog::Open(EAppMsgType::YesNo, LevelInstanceBreakWarning) == EAppReturnType::Yes)
 						{
-							const FText LevelInstanceBreakWarning = LOCTEXT("BreakingLevelInstance", "You are about to break the level instance. This action cannot be undone. Are you sure ?");
-							if (FMessageDialog::Open(EAppMsgType::YesNo, LevelInstanceBreakWarning) == EAppReturnType::Yes)
+							for (ILevelInstanceInterface* LevelInstance : BreakableLevelInstances)
 							{
-								LevelInstanceSubsystem->BreakLevelInstance(ContextLevelInstance, BreakLevels);
+								LevelInstanceSubsystem->BreakLevelInstance(LevelInstance, BreakLevels);
 							}
-							return FReply::Handled();
-						})
-					.Text(LOCTEXT("BreakLevelInstances_BreakLevelInstanceButton", "Break Level Instance"))
+						}
+						return FReply::Handled();
+					})
+					.Text(LOCTEXT("BreakLevelInstances_BreakLevelInstanceButton", "Break Level Instance(s)"))
 				];
 
 			Section.AddEntry(FToolMenuEntry::InitWidget("SetBreakLevels", MenuWidget, FText::GetEmpty(), false));
 		}
 	}
 
-	void CreateBreakMenu(UToolMenu* Menu, AActor* ContextActor)
+	void CreateBreakMenu(UToolMenu* Menu)
 	{
-		check(ContextActor);
-
-		if (ULevelInstanceSubsystem* LevelInstanceSubsystem = ContextActor->GetWorld()->GetSubsystem<ULevelInstanceSubsystem>())
+		if(ULevelInstanceSubsystem* LevelInstanceSubsystem = GEditor->GetEditorWorldContext().World()->GetSubsystem<ULevelInstanceSubsystem>())
 		{
-			ILevelInstanceInterface* ContextLevelInstance = nullptr;
-
-			// Find the top level LevelInstance
-			LevelInstanceSubsystem->ForEachLevelInstanceAncestorsAndSelf(ContextActor, [LevelInstanceSubsystem, ContextActor, &ContextLevelInstance](ILevelInstanceInterface* Ancestor)
+			TArray<ILevelInstanceInterface*> BreakableLevelInstances;
+			if (GEditor->GetSelectedActorCount() > 0)
 			{
-				if (CastChecked<AActor>(Ancestor)->GetLevel() == ContextActor->GetWorld()->GetCurrentLevel())
+				for (FSelectionIterator It(GEditor->GetSelectedActorIterator()); It; ++It)
 				{
-					ContextLevelInstance = Ancestor;
-					return false;
+					if (ILevelInstanceInterface* LevelInstance = Cast<ILevelInstanceInterface>(*It))
+					{
+						if (IsExperimentalSettingEnabled(LevelInstance) && !LevelInstanceSubsystem->IsEditingLevelInstance(LevelInstance) && !LevelInstanceSubsystem->LevelInstanceHasLevelScriptBlueprint(LevelInstance))
+						{
+							BreakableLevelInstances.Add(LevelInstance);
+						}	
+					}
 				}
-				return true;
-			});
+			}
 
-			if (ContextLevelInstance && IsExperimentalSettingEnabled(ContextLevelInstance) && !LevelInstanceSubsystem->IsEditingLevelInstance(ContextLevelInstance) && !LevelInstanceSubsystem->LevelInstanceHasLevelScriptBlueprint(ContextLevelInstance))
+			if (BreakableLevelInstances.Num())
 			{
 				FToolMenuSection& Section = CreateLevelSection(Menu);
-
 				Section.AddSubMenu(
 					"BreakLevelInstances",
 					LOCTEXT("BreakLevelInstances", "Break..."),
 					TAttribute<FText>(),
-					FNewToolMenuDelegate::CreateStatic(&CreateBreakSubMenu, ContextLevelInstance)
+					FNewToolMenuDelegate::CreateStatic(&CreateBreakSubMenu, BreakableLevelInstances)
 				);
 			}
 		}
+
 	}
 
 	void CreatePackedBlueprintMenu(UToolMenu* Menu, AActor* ContextActor)
@@ -743,7 +744,6 @@ void FLevelInstanceEditorModule::ExtendContextMenu()
 				LevelInstanceMenuUtils::CreateEditMenu(ToolMenu, ContextActor);
 				LevelInstanceMenuUtils::CreateAddDataLayerSupportMenu(ToolMenu, ContextActor);
 				LevelInstanceMenuUtils::CreateCommitDiscardMenu(ToolMenu, ContextActor);
-				LevelInstanceMenuUtils::CreateBreakMenu(ToolMenu, ContextActor);
 				LevelInstanceMenuUtils::CreatePackedBlueprintMenu(ToolMenu, ContextActor);
 				LevelInstanceMenuUtils::CreateSetCurrentMenu(ToolMenu, ContextActor);
 			}
@@ -751,6 +751,7 @@ void FLevelInstanceEditorModule::ExtendContextMenu()
 			LevelInstanceMenuUtils::CreateMoveSelectionToMenu(ToolMenu);
 		}
 
+		LevelInstanceMenuUtils::CreateBreakMenu(ToolMenu);
 		LevelInstanceMenuUtils::CreateCreateMenu(ToolMenu);
 	};
 
