LevelInstance: Fix Placement Mode not showing Level Instance

diff --git a/Engine/Source/Editor/LevelInstanceEditor/Private/LevelInstanceActorFactory.cpp b/Engine/Source/Editor/LevelInstanceEditor/Private/LevelInstanceActorFactory.cpp
--- a/Engine/Source/Editor/LevelInstanceEditor/Private/LevelInstanceActorFactory.cpp
+++ b/Engine/Source/Editor/LevelInstanceEditor/Private/LevelInstanceActorFactory.cpp
@@ -26,7 +26,8 @@ bool ULevelInstanceActorFactory::CanCreateActorFrom(const FAssetData& AssetData,
 		return false;
 	}
 
-	if (!AssetData.IsValid() || !AssetData.IsInstanceOf(UWorld::StaticClass()))
+	// If asset is valid it needs to be of type: UWorld
+	if (AssetData.IsValid() && !AssetData.IsInstanceOf(UWorld::StaticClass()))
 	{
 		OutErrorMsg = NSLOCTEXT("LevelInstanceActorFactory", "NoWorld", "A valid world must be specified.");
 		return false;
