Changed ULevelEditorContextMenuContext's HitProxyActor from TObjectPtr to TWeakObjectPtr to avoid getting reports of a reference to a stale world when entering level instance edit mode.

diff --git a/Engine/Source/Editor/LevelEditor/Public/LevelEditorMenuContext.h b/Engine/Source/Editor/LevelEditor/Public/LevelEditorMenuContext.h
--- a/Engine/Source/Editor/LevelEditor/Public/LevelEditorMenuContext.h
+++ b/Engine/Source/Editor/LevelEditor/Public/LevelEditorMenuContext.h
@@ -61,7 +61,7 @@ public:
 
 	/** If the ContextType is Viewport this property can be set to the HitProxy actor that triggered the ContextMenu. */
 	UPROPERTY(VisibleAnywhere, BlueprintReadOnly, Category="LevelEditor|Menu")
-	TObjectPtr<AActor> HitProxyActor = nullptr;
+	TWeakObjectPtr<AActor> HitProxyActor = nullptr;
 
 	UFUNCTION(BlueprintCallable, Category = "LevelEditor | Menu", DisplayName="Get Hit Proxy Element", meta=(ScriptName="GetHitProxyElement"))
 	FScriptTypedElementHandle GetScriptHitProxyElement();
diff --git a/Engine/Source/Editor/LevelInstanceEditor/Private/LevelInstanceEditorModule.cpp b/Engine/Source/Editor/LevelInstanceEditor/Private/LevelInstanceEditorModule.cpp
--- a/Engine/Source/Editor/LevelInstanceEditor/Private/LevelInstanceEditorModule.cpp
+++ b/Engine/Source/Editor/LevelInstanceEditor/Private/LevelInstanceEditorModule.cpp
@@ -733,7 +733,7 @@ void FLevelInstanceEditorModule::ExtendContextMenu()
 		{
 			// Use the actor under the cursor if available (e.g. right-click menu).
 			// Otherwise use the first selected actor if there's one (e.g. Actor pulldown menu or outliner).
-			AActor* ContextActor = LevelEditorMenuContext->HitProxyActor;
+			AActor* ContextActor = LevelEditorMenuContext->HitProxyActor.Get();
 			if (!ContextActor && GEditor->GetSelectedActorCount() != 0)
 			{
 				ContextActor = Cast<AActor>(GEditor->GetSelectedActors()->GetSelectedObject(0));
