LevelInstance: Fix EditingLevelInstance show flag not being properly updated

diff --git a/Engine/Source/Editor/LevelInstanceEditor/Private/LevelInstanceEditorMode.cpp b/Engine/Source/Editor/LevelInstanceEditor/Private/LevelInstanceEditorMode.cpp
--- a/Engine/Source/Editor/LevelInstanceEditor/Private/LevelInstanceEditorMode.cpp
+++ b/Engine/Source/Editor/LevelInstanceEditor/Private/LevelInstanceEditorMode.cpp
@@ -48,7 +48,10 @@ void ULevelInstanceEditorMode::UpdateEngineShowFlags()
 		{
 			if(ULevelInstanceSubsystem* LevelInstanceSubsystem = LevelVC->GetWorld()->GetSubsystem<ULevelInstanceSubsystem>())
 			{
-				LevelVC->EngineShowFlags.EditingLevelInstance = !!LevelInstanceSubsystem->GetEditingLevelInstance();
+				const bool bEditingLevelInstance = !!LevelInstanceSubsystem->GetEditingLevelInstance();
+				// Make sure we update both Game/Editor showflags
+				LevelVC->EngineShowFlags.EditingLevelInstance = bEditingLevelInstance;
+				LevelVC->LastEngineShowFlags.EditingLevelInstance = bEditingLevelInstance;
 			}
 		}
 	}
