World Partition: added option to add partitioned level instance streaming via the content browser, and remove the per-instance actor support.

#rnx

diff --git a/Engine/Source/Editor/LevelInstanceEditor/Private/LevelInstanceEditorModule.cpp b/Engine/Source/Editor/LevelInstanceEditor/Private/LevelInstanceEditorModule.cpp
--- a/Engine/Source/Editor/LevelInstanceEditor/Private/LevelInstanceEditorModule.cpp
+++ b/Engine/Source/Editor/LevelInstanceEditor/Private/LevelInstanceEditorModule.cpp
@@ -11,6 +11,7 @@
 #include "Editor.h"
 #include "EditorModeManager.h"
 #include "EditorModeRegistry.h"
+#include "FileHelpers.h"
 #include "LevelInstanceEditorMode.h"
 #include "LevelInstanceEditorModeCommands.h"
 #include "LevelEditorMenuContext.h"
@@ -39,6 +40,7 @@
 #include "SNewLevelInstanceDialog.h"
 #include "MessageLogModule.h"
 #include "Settings/EditorExperimentalSettings.h"
+#include "WorldPartition/WorldPartitionConverter.h"
 
 IMPLEMENT_MODULE( FLevelInstanceEditorModule, LevelInstanceEditor );
 
@@ -152,32 +154,6 @@ namespace LevelInstanceMenuUtils
 		}
 	}
 
-	void CreateAddDataLayerSupportMenu(UToolMenu* Menu, AActor* ContextActor)
-	{
-		FText EntryDesc;
-		ILevelInstanceInterface* LevelInstance = Cast<ILevelInstanceInterface>(ContextActor);
-		ULevelInstanceSubsystem* LevelInstanceSubsystem = ContextActor->GetWorld()->GetSubsystem<ULevelInstanceSubsystem>();
-		if (!LevelInstance || 
-			!LevelInstanceSubsystem ||
-			!LevelInstance->IsWorldAssetValid() ||
-			!LevelInstance->CanEnterEdit(&EntryDesc) ||
-			!ULevelInstanceSubsystem::CanUsePackage(*LevelInstance->GetWorldAssetPackage()) ||
-			ULevel::GetPartitionedLevelCanBeUsedByLevelInstanceFromPackage(*LevelInstance->GetWorldAssetPackage()))
-		{
-			return;
-		}
-
-		FToolUIAction LevelInstanceEditAction;
-		LevelInstanceEditAction.ExecuteAction.BindLambda([LevelInstanceSubsystem, LevelInstance, ContextActor](const FToolMenuContext&)
-		{
-			LevelInstanceSubsystem->AddDataLayerSupport(LevelInstance);
-		});
-
-		AActor* LevelInstanceActor = CastChecked<AActor>(LevelInstance);
-		EntryDesc = FText::Format(LOCTEXT("LevelInstanceName", "{0}:{1}"), FText::FromString(LevelInstanceActor->GetActorLabel()), FText::FromString(LevelInstance->GetWorldAssetPackage()));
-		CreateLevelSection(Menu).AddMenuEntry(NAME_None, LOCTEXT("AddDataLayerSupport", "Add Data Layer Support"), EntryDesc, FSlateIcon(), LevelInstanceEditAction);
-	}
-		
 	void CreateCommitDiscardMenu(UToolMenu* Menu, AActor* ContextActor)
 	{
 		ILevelInstanceInterface* LevelInstanceEdit = nullptr;
@@ -596,6 +572,81 @@ namespace LevelInstanceMenuUtils
 			TAttribute<FSlateIcon>(),
 			UIAction);
 	}
+
+	void AddPartitionedStreamingSupportFromWorld(UWorld* WorldAsset)
+	{
+		bool bSuccess = false;
+
+		UWorld* World = GEditor->GetEditorWorldContext().World();
+		if (ULevelInstanceSubsystem* LevelInstanceSubsystem = World->GetSubsystem<ULevelInstanceSubsystem>())
+		{
+			LevelInstanceSubsystem->ResetLoadersForWorldAsset(WorldAsset->GetPackage()->GetName());
+
+			FWorldPartitionConverter::FParameters Parameters;
+			Parameters.bConvertSubLevels = false;
+			Parameters.bEnableStreaming = false;
+			Parameters.bUseActorFolders = true;
+
+			if (FWorldPartitionConverter::Convert(WorldAsset, Parameters))
+			{
+				TArray<UPackage*> PackagesToSave = WorldAsset->PersistentLevel->GetLoadedExternalObjectPackages();
+				TSet<UPackage*> PackagesToSaveSet(PackagesToSave);
+
+				PackagesToSaveSet.Add(WorldAsset->GetPackage());
+
+				const bool bPromptUserToSave = false;
+				const bool bSaveMapPackages = true;
+				const bool bSaveContentPackages = true;
+				const bool bFastSave = false;
+				const bool bNotifyNoPackagesSaved = false;
+				const bool bCanBeDeclined = true;
+
+				if (FEditorFileUtils::SaveDirtyPackages(bPromptUserToSave, bSaveMapPackages, bSaveContentPackages, bFastSave, bNotifyNoPackagesSaved, bCanBeDeclined, nullptr, [&PackagesToSaveSet](UPackage* PackageToSave) { return !PackagesToSaveSet.Contains(PackageToSave); }))
+				{
+					bSuccess = true;
+					for (UPackage* PackageToSave : PackagesToSave)
+					{
+						if (PackageToSave->IsDirty())
+						{
+							bSuccess = false;
+							break;
+						}
+					}
+				}
+			}
+		}
+
+		if (!bSuccess)
+		{
+			FMessageDialog::Open(EAppMsgType::Ok, LOCTEXT("AddPartitionedLevelInstanceStreamingSupportError", "An error occured when adding partitioned level instance streaming support, check logs for details.."));
+		}
+	}
+
+	void AddPartitionedStreamingSupportFromMenu(UToolMenu* Menu, UWorld* WorldAsset)
+	{
+		FName WorldAssetName = WorldAsset->GetPackage()->GetFName();
+		if (!ULevel::GetIsLevelPartitionedFromPackage(WorldAssetName) || !ULevel::GetPartitionedLevelCanBeUsedByLevelInstanceFromPackage(WorldAssetName))
+		{
+			FToolMenuSection& Section = CreateLevelSection(Menu);
+			FToolUIAction UIAction;
+			UIAction.ExecuteAction.BindLambda([WorldAsset](const FToolMenuContext& MenuContext)
+			{
+				AddPartitionedStreamingSupportFromWorld(WorldAsset);
+			});
+			UIAction.CanExecuteAction.BindLambda([WorldAsset](const FToolMenuContext&)
+			{
+				return !WorldAsset->GetStreamingLevels().Num() && (WorldAsset->WorldType == EWorldType::Inactive);
+			});
+
+			Section.AddMenuEntry(
+				"AddPartitionedStreamingSupportFromMenu",
+				LOCTEXT("AddPartitionedStreamingSupportFromMenu", "Add Partitioned Streaming Support"),
+				TAttribute<FText>(),
+				TAttribute<FSlateIcon>(),
+				UIAction
+			);
+		}
+	}
 };
 
 void FLevelInstanceEditorModule::StartupModule()
@@ -742,7 +793,6 @@ void FLevelInstanceEditorModule::ExtendContextMenu()
 			if (ContextActor)
 			{
 				LevelInstanceMenuUtils::CreateEditMenu(ToolMenu, ContextActor);
-				LevelInstanceMenuUtils::CreateAddDataLayerSupportMenu(ToolMenu, ContextActor);
 				LevelInstanceMenuUtils::CreateCommitDiscardMenu(ToolMenu, ContextActor);
 				LevelInstanceMenuUtils::CreatePackedBlueprintMenu(ToolMenu, ContextActor);
 				LevelInstanceMenuUtils::CreateSetCurrentMenu(ToolMenu, ContextActor);
@@ -786,6 +836,7 @@ void FLevelInstanceEditorModule::ExtendContextMenu()
 					if (UWorld* WorldAsset = Cast<UWorld>(AssetMenuContext->SelectedObjects[0].Get()))
 					{
 						LevelInstanceMenuUtils::CreateBlueprintFromMenu(ToolMenu, WorldAsset);
+						LevelInstanceMenuUtils::AddPartitionedStreamingSupportFromMenu(ToolMenu, WorldAsset);
 					}
 				}
 			}
diff --git a/Engine/Source/Runtime/Engine/Private/LevelInstance/LevelInstanceSubsystem.cpp b/Engine/Source/Runtime/Engine/Private/LevelInstance/LevelInstanceSubsystem.cpp
--- a/Engine/Source/Runtime/Engine/Private/LevelInstance/LevelInstanceSubsystem.cpp
+++ b/Engine/Source/Runtime/Engine/Private/LevelInstance/LevelInstanceSubsystem.cpp
@@ -24,7 +24,6 @@
 #include "WorldPartition/WorldPartitionSubsystem.h"
 #include "WorldPartition/DataLayer/DataLayerSubsystem.h"
 #include "WorldPartition/DataLayer/DataLayerInstanceWithAsset.h"
-#include "WorldPartition/WorldPartitionConverter.h"
 #include "WorldPartition/WorldPartitionMiniMap.h"
 #include "Misc/ScopedSlowTask.h"
 #include "Misc/ITransaction.h"
@@ -1950,15 +1949,20 @@ bool ULevelInstanceSubsystem::EditLevelInstanceInternal(ILevelInstanceInterface*
 	// When editing a Level Instance, push a new empty actor editor context
 	UActorEditorContextSubsystem::Get()->PushContext();
 
+	ResetLoadersForWorldAsset(LevelInstance->GetWorldAsset().GetLongPackageName());
+
+	return true;
+}
+
+void ULevelInstanceSubsystem::ResetLoadersForWorldAsset(const FString& WorldAsset)
+{
 	for (const auto& [LevelInstanceID, LoadedLevelInstance] : LoadedLevelInstances)
 	{
-		if (LoadedLevelInstance.LevelStreaming && LoadedLevelInstance.LevelStreaming->PackageNameToLoad.ToString() == LevelInstance->GetWorldAsset().GetLongPackageName())
+		if (LoadedLevelInstance.LevelStreaming && LoadedLevelInstance.LevelStreaming->PackageNameToLoad.ToString() == WorldAsset)
 		{
 			LoadedLevelInstance.LevelStreaming->ResetLevelInstanceLoaders();
 		}
 	}
-
-	return true;
 }
 
 bool ULevelInstanceSubsystem::CommitLevelInstance(ILevelInstanceInterface* LevelInstance, bool bDiscardEdits, TSet<FName>* DirtyPackages)
@@ -1976,39 +1980,6 @@ bool ULevelInstanceSubsystem::CommitLevelInstance(ILevelInstanceInterface* Level
 	return false;
 }
 
-bool ULevelInstanceSubsystem::AddDataLayerSupport(ILevelInstanceInterface* LevelInstance)
-{
-	check(LevelInstance);
-	
-	bool bSuccess = false;
-	EditLevelInstance(LevelInstance);
-	
-	if (GetEditingLevelInstance())
-	{
-		UWorld* World = Cast<UWorld>(LevelInstance->GetWorldAsset().Get());
-		if (World && !World->IsPartitionedWorld())
-		{
-			FWorldPartitionConverter::FParameters Parameters;
-			Parameters.bConvertSubLevels = false;
-			Parameters.bEnableStreaming = false;
-			Parameters.bUseActorFolders = true;
-			if (FWorldPartitionConverter::Convert(World, Parameters))
-			{
-				UWorldPartition* WorldPartition = World->GetWorldPartition();
-				check(WorldPartition);
-				check(WorldPartition->CanBeUsedByLevelInstance());
-				check(World->PersistentLevel->IsUsingActorFolders());
-				bSuccess = true;
-			}
-		}
-
-		CommitLevelInstance(LevelInstance, !bSuccess);
-	}
-
-	UE_CLOG(!bSuccess, LogLevelInstance, Error, TEXT("Failed to convert level %s for Level Instance Data Layer support."), *LevelInstance->GetWorldAssetPackage());
-	return bSuccess;
-}
-
 bool ULevelInstanceSubsystem::CommitLevelInstanceInternal(TUniquePtr<FLevelInstanceEdit>& InLevelInstanceEdit, bool bDiscardEdits, bool bDiscardOnFailure, TSet<FName>* DirtyPackages)
 {
 	TGuardValue<bool> CommitScope(bIsCommittingLevelInstance, true);
diff --git a/Engine/Source/Runtime/Engine/Public/LevelInstance/LevelInstanceSubsystem.h b/Engine/Source/Runtime/Engine/Public/LevelInstance/LevelInstanceSubsystem.h
--- a/Engine/Source/Runtime/Engine/Public/LevelInstance/LevelInstanceSubsystem.h
+++ b/Engine/Source/Runtime/Engine/Public/LevelInstance/LevelInstanceSubsystem.h
@@ -53,6 +53,7 @@ public:
 	bool IsLoaded(const ILevelInstanceInterface* LevelInstance) const;
 	void ForEachLevelInstanceAncestorsAndSelf(AActor* Actor, TFunctionRef<bool(ILevelInstanceInterface*)> Operation) const;
 	ULevelStreamingLevelInstance* GetLevelInstanceLevelStreaming(const ILevelInstanceInterface* LevelInstance) const;
+	void ResetLoadersForWorldAsset(const FString& WorldAsset);
 
 #if WITH_EDITOR
 	void Tick();
@@ -67,7 +68,6 @@ public:
 	bool CanCommitLevelInstance(const ILevelInstanceInterface* LevelInstance, bool bDiscardEdits = false, FText* OutReason = nullptr) const;
 	void EditLevelInstance(ILevelInstanceInterface* LevelInstance, TWeakObjectPtr<AActor> ContextActorPtr = nullptr);
 	bool CommitLevelInstance(ILevelInstanceInterface* LevelInstance, bool bDiscardEdits = false, TSet<FName>* DirtyPackages = nullptr);
-	bool AddDataLayerSupport(ILevelInstanceInterface* LevelInstance);
 	bool IsEditingLevelInstanceDirty(const ILevelInstanceInterface* LevelInstance) const;
 	bool IsEditingLevelInstance(const ILevelInstanceInterface* LevelInstance) const { return GetLevelInstanceEdit(LevelInstance) != nullptr; }
 	
