World Partition: added better logs and user messages when adding partitioned level instance streaming support fails.

#rnx

diff --git a/Engine/Source/Editor/LevelInstanceEditor/Private/LevelInstanceEditorModule.cpp b/Engine/Source/Editor/LevelInstanceEditor/Private/LevelInstanceEditorModule.cpp
--- a/Engine/Source/Editor/LevelInstanceEditor/Private/LevelInstanceEditorModule.cpp
+++ b/Engine/Source/Editor/LevelInstanceEditor/Private/LevelInstanceEditorModule.cpp
@@ -46,6 +46,8 @@ IMPLEMENT_MODULE( FLevelInstanceEditorModule, LevelInstanceEditor );
 
 #define LOCTEXT_NAMESPACE "LevelInstanceEditor"
 
+DEFINE_LOG_CATEGORY_STATIC(LogLevelInstanceEditor, Log, All);
+
 namespace LevelInstanceMenuUtils
 {
 	bool IsExperimentalSettingEnabled(ILevelInstanceInterface* LevelInstance)
@@ -575,8 +577,19 @@ namespace LevelInstanceMenuUtils
 
 	void AddPartitionedStreamingSupportFromWorld(UWorld* WorldAsset)
 	{
-		bool bSuccess = false;
+		if (WorldAsset->GetStreamingLevels().Num())
+		{
+			FMessageDialog::Open(EAppMsgType::Ok, LOCTEXT("AddPartitionedLevelInstanceStreamingSupportError_SubLevels", "Cannot convert this world has it contains sublevels."));
+			return;
+		}
+
+		if (WorldAsset->WorldType != EWorldType::Inactive)
+		{
+			FMessageDialog::Open(EAppMsgType::Ok, LOCTEXT("AddPartitionedLevelInstanceStreamingSupportError_Loaded", "Cannot convert this world has it's already loaded in the editor."));
+			return;
+		}
 
+		bool bSuccess = false;
 		UWorld* World = GEditor->GetEditorWorldContext().World();
 		if (ULevelInstanceSubsystem* LevelInstanceSubsystem = World->GetSubsystem<ULevelInstanceSubsystem>())
 		{
@@ -608,6 +621,7 @@ namespace LevelInstanceMenuUtils
 					{
 						if (PackageToSave->IsDirty())
 						{
+							UE_LOG(LogLevelInstanceEditor, Error, TEXT("Package '%s' failed to save"), *PackageToSave->GetName());
 							bSuccess = false;
 							break;
 						}
@@ -633,10 +647,6 @@ namespace LevelInstanceMenuUtils
 			{
 				AddPartitionedStreamingSupportFromWorld(WorldAsset);
 			});
-			UIAction.CanExecuteAction.BindLambda([WorldAsset](const FToolMenuContext&)
-			{
-				return !WorldAsset->GetStreamingLevels().Num() && (WorldAsset->WorldType == EWorldType::Inactive);
-			});
 
 			Section.AddMenuEntry(
 				"AddPartitionedStreamingSupportFromMenu",
