LevelInstandeSubsystem fixes:
- Unable  to overwrite existing level
- Deleting level doesn't unload loaded instances

diff --git a/Engine/Source/Editor/LevelInstanceEditor/Private/LevelInstanceEditorModule.cpp b/Engine/Source/Editor/LevelInstanceEditor/Private/LevelInstanceEditorModule.cpp
--- a/Engine/Source/Editor/LevelInstanceEditor/Private/LevelInstanceEditorModule.cpp
+++ b/Engine/Source/Editor/LevelInstanceEditor/Private/LevelInstanceEditorModule.cpp
@@ -591,40 +591,37 @@ namespace LevelInstanceMenuUtils
 
 		bool bSuccess = false;
 		UWorld* World = GEditor->GetEditorWorldContext().World();
-		if (ULevelInstanceSubsystem* LevelInstanceSubsystem = World->GetSubsystem<ULevelInstanceSubsystem>())
-		{
-			LevelInstanceSubsystem->ResetLoadersForWorldAsset(WorldAsset->GetPackage()->GetName());
+		ULevelInstanceSubsystem::ResetLoadersForWorldAsset(WorldAsset->GetPackage()->GetName());
 
-			FWorldPartitionConverter::FParameters Parameters;
-			Parameters.bConvertSubLevels = false;
-			Parameters.bEnableStreaming = false;
-			Parameters.bUseActorFolders = true;
+		FWorldPartitionConverter::FParameters Parameters;
+		Parameters.bConvertSubLevels = false;
+		Parameters.bEnableStreaming = false;
+		Parameters.bUseActorFolders = true;
 
-			if (FWorldPartitionConverter::Convert(WorldAsset, Parameters))
-			{
-				TArray<UPackage*> PackagesToSave = WorldAsset->PersistentLevel->GetLoadedExternalObjectPackages();
-				TSet<UPackage*> PackagesToSaveSet(PackagesToSave);
+		if (FWorldPartitionConverter::Convert(WorldAsset, Parameters))
+		{
+			TArray<UPackage*> PackagesToSave = WorldAsset->PersistentLevel->GetLoadedExternalObjectPackages();
+			TSet<UPackage*> PackagesToSaveSet(PackagesToSave);
 
-				PackagesToSaveSet.Add(WorldAsset->GetPackage());
+			PackagesToSaveSet.Add(WorldAsset->GetPackage());
 
-				const bool bPromptUserToSave = false;
-				const bool bSaveMapPackages = true;
-				const bool bSaveContentPackages = true;
-				const bool bFastSave = false;
-				const bool bNotifyNoPackagesSaved = false;
-				const bool bCanBeDeclined = true;
+			const bool bPromptUserToSave = false;
+			const bool bSaveMapPackages = true;
+			const bool bSaveContentPackages = true;
+			const bool bFastSave = false;
+			const bool bNotifyNoPackagesSaved = false;
+			const bool bCanBeDeclined = true;
 
-				if (FEditorFileUtils::SaveDirtyPackages(bPromptUserToSave, bSaveMapPackages, bSaveContentPackages, bFastSave, bNotifyNoPackagesSaved, bCanBeDeclined, nullptr, [&PackagesToSaveSet](UPackage* PackageToSave) { return !PackagesToSaveSet.Contains(PackageToSave); }))
+			if (FEditorFileUtils::SaveDirtyPackages(bPromptUserToSave, bSaveMapPackages, bSaveContentPackages, bFastSave, bNotifyNoPackagesSaved, bCanBeDeclined, nullptr, [&PackagesToSaveSet](UPackage* PackageToSave) { return !PackagesToSaveSet.Contains(PackageToSave); }))
+			{
+				bSuccess = true;
+				for (UPackage* PackageToSave : PackagesToSave)
 				{
-					bSuccess = true;
-					for (UPackage* PackageToSave : PackagesToSave)
+					if (PackageToSave->IsDirty())
 					{
-						if (PackageToSave->IsDirty())
-						{
-							UE_LOG(LogLevelInstanceEditor, Error, TEXT("Package '%s' failed to save"), *PackageToSave->GetName());
-							bSuccess = false;
-							break;
-						}
+						UE_LOG(LogLevelInstanceEditor, Error, TEXT("Package '%s' failed to save"), *PackageToSave->GetName());
+						bSuccess = false;
+						break;
 					}
 				}
 			}
diff --git a/Engine/Source/Editor/UnrealEd/Private/FileHelpers.cpp b/Engine/Source/Editor/UnrealEd/Private/FileHelpers.cpp
--- a/Engine/Source/Editor/UnrealEd/Private/FileHelpers.cpp
+++ b/Engine/Source/Editor/UnrealEd/Private/FileHelpers.cpp
@@ -80,6 +80,7 @@
 #include "WorldPartition/WorldPartitionEditorPerProjectUserSettings.h"
 #include "WorldPartition/HLOD/HLODLayer.h"
 #include "WorldPartition/LoaderAdapter/LoaderAdapterShape.h"
+#include "LevelInstance/LevelInstanceSubsystem.h"
 #include "PackageSourceControlHelper.h"
 #include "ActorFolder.h"
 #include "InterchangeManager.h"
@@ -771,6 +772,9 @@ static bool SaveWorld(UWorld* World,
 		{
 			if (bPackageNeedsRename)
 			{
+				// Reset Loaders before deleting existing packages
+				ULevelInstanceSubsystem::ResetLoadersForWorldAsset(NewPackageName);
+
 				// Delete files at destination
 				if (!DeleteExistingMapPackages(NewPackageName))
 				{
diff --git a/Engine/Source/Runtime/Engine/Private/LevelInstance/LevelInstanceSubsystem.cpp b/Engine/Source/Runtime/Engine/Private/LevelInstance/LevelInstanceSubsystem.cpp
--- a/Engine/Source/Runtime/Engine/Private/LevelInstance/LevelInstanceSubsystem.cpp
+++ b/Engine/Source/Runtime/Engine/Private/LevelInstance/LevelInstanceSubsystem.cpp
@@ -90,10 +90,18 @@ void ULevelInstanceSubsystem::Initialize(FSubsystemCollectionBase& Collection)
 	if (GEditor)
 	{
 		ILevelInstanceEditorModule& EditorModule = FModuleManager::LoadModuleChecked<ILevelInstanceEditorModule>("LevelInstanceEditor");
+		FEditorDelegates::OnAssetsPreDelete.AddUObject(this, &ULevelInstanceSubsystem::OnAssetsPreDelete);
 	}
 #endif
 }
 
+void ULevelInstanceSubsystem::Deinitialize()
+{
+#if WITH_EDITOR
+	FEditorDelegates::OnAssetsPreDelete.RemoveAll(this);
+#endif
+}
+
 bool ULevelInstanceSubsystem::DoesSupportWorldType(EWorldType::Type WorldType) const
 {
 	return Super::DoesSupportWorldType(WorldType) || WorldType == EWorldType::EditorPreview || WorldType == EWorldType::Inactive;
@@ -286,6 +294,25 @@ void ULevelInstanceSubsystem::RegisterLoadedLevelStreamingLevelInstance(ULevelSt
 }
 
 #if WITH_EDITOR
+
+void ULevelInstanceSubsystem::OnAssetsPreDelete(const TArray<UObject*>& Objects)
+{
+	for (UObject* Object : Objects)
+	{
+		if (IsValid(Object))
+		{
+			if (UPackage* Package = Object->GetPackage())
+			{
+				TArray<ILevelInstanceInterface*> LevelInstances = GetLevelInstances(Package->GetLoadedPath().GetPackageName());
+				for (ILevelInstanceInterface* LevelInstancePtr : LevelInstances)
+				{
+					UnloadLevelInstance(LevelInstancePtr->GetLevelInstanceID());
+				}
+			}
+		}
+	}
+}
+
 void ULevelInstanceSubsystem::RegisterLoadedLevelStreamingLevelInstanceEditor(ULevelStreamingLevelInstanceEditor* LevelStreaming)
 {
 	if (!bIsCreatingLevelInstance)
@@ -1954,6 +1981,21 @@ bool ULevelInstanceSubsystem::EditLevelInstanceInternal(ILevelInstanceInterface*
 }
 
 void ULevelInstanceSubsystem::ResetLoadersForWorldAsset(const FString& WorldAsset)
+{
+	for (TObjectIterator<UWorld> It(RF_ClassDefaultObject | RF_ArchetypeObject, true); It; ++It)
+	{
+		UWorld* CurrentWorld = *It;
+		if (IsValid(CurrentWorld))
+		{
+			if (ULevelInstanceSubsystem* LevelInstanceSubsystem = CurrentWorld->GetSubsystem<ULevelInstanceSubsystem>())
+			{
+				LevelInstanceSubsystem->ResetLoadersForWorldAssetInternal(WorldAsset);
+			}
+		}
+	}
+}
+
+void ULevelInstanceSubsystem::ResetLoadersForWorldAssetInternal(const FString& WorldAsset)
 {
 	for (const auto& [LevelInstanceID, LoadedLevelInstance] : LoadedLevelInstances)
 	{
diff --git a/Engine/Source/Runtime/Engine/Public/LevelInstance/LevelInstanceSubsystem.h b/Engine/Source/Runtime/Engine/Public/LevelInstance/LevelInstanceSubsystem.h
--- a/Engine/Source/Runtime/Engine/Public/LevelInstance/LevelInstanceSubsystem.h
+++ b/Engine/Source/Runtime/Engine/Public/LevelInstance/LevelInstanceSubsystem.h
@@ -37,6 +37,7 @@ public:
 		
 	//~ Begin USubsystem Interface.
 	virtual void Initialize(FSubsystemCollectionBase& Collection) override;
+	virtual void Deinitialize();
 	virtual bool DoesSupportWorldType(EWorldType::Type WorldType) const override;
 	//~ End USubsystem Interface.
 	
@@ -53,7 +54,6 @@ public:
 	bool IsLoaded(const ILevelInstanceInterface* LevelInstance) const;
 	void ForEachLevelInstanceAncestorsAndSelf(AActor* Actor, TFunctionRef<bool(ILevelInstanceInterface*)> Operation) const;
 	ULevelStreamingLevelInstance* GetLevelInstanceLevelStreaming(const ILevelInstanceInterface* LevelInstance) const;
-	void ResetLoadersForWorldAsset(const FString& WorldAsset);
 
 #if WITH_EDITOR
 	void Tick();
@@ -122,6 +122,8 @@ public:
 	/** Editor-only event triggered when level instance is committed with changes */
 	DECLARE_EVENT_OneParam(ULevelInstanceSubsystem, FLevelInstanceChanged, FName);
 	FLevelInstanceChanged& OnLevelInstanceChanged() { return LevelInstanceChangedEvent; }
+
+	static void ResetLoadersForWorldAsset(const FString& WorldAsset);
 #endif
 
 private:
@@ -134,6 +136,9 @@ private:
 	void RegisterLoadedLevelStreamingLevelInstance(ULevelStreamingLevelInstance* LevelStreaming);
 
 #if WITH_EDITOR
+	void ResetLoadersForWorldAssetInternal(const FString& WorldAsset);
+	void OnAssetsPreDelete(const TArray<UObject*>& Objects);
+
 	void RegisterLoadedLevelStreamingLevelInstanceEditor(ULevelStreamingLevelInstanceEditor* LevelStreaming);
 
 	void OnEditChild(const FLevelInstanceID& LevelInstanceID);
