Level Instance Pivot Picker:
- Preview picker location when changing it from the details panel
- Use Actor Picker which isn't restricted to the Level Instance actors

diff --git a/Engine/Source/Editor/LevelInstanceEditor/LevelInstanceEditor.Build.cs b/Engine/Source/Editor/LevelInstanceEditor/LevelInstanceEditor.Build.cs
--- a/Engine/Source/Editor/LevelInstanceEditor/LevelInstanceEditor.Build.cs
+++ b/Engine/Source/Editor/LevelInstanceEditor/LevelInstanceEditor.Build.cs
@@ -31,7 +31,8 @@ public class LevelInstanceEditor : ModuleRules
 				"ClassViewer",
 				"MessageLog",
 				"EditorWidgets",
-				"DeveloperSettings"
+				"DeveloperSettings",
+				"SceneOutliner"
 			}
 		);
 
diff --git a/Engine/Source/Editor/LevelInstanceEditor/Private/LevelInstancePivotDetails.cpp b/Engine/Source/Editor/LevelInstanceEditor/Private/LevelInstancePivotDetails.cpp
--- a/Engine/Source/Editor/LevelInstanceEditor/Private/LevelInstancePivotDetails.cpp
+++ b/Engine/Source/Editor/LevelInstanceEditor/Private/LevelInstancePivotDetails.cpp
@@ -6,7 +6,6 @@
 #include "DetailWidgetRow.h"
 #include "Engine/Level.h"
 #include "Widgets/Input/SButton.h"
-#include "Widgets/Input/SComboBox.h"
 #include "Widgets/Input/SMultiLineEditableTextBox.h"
 #include "SEnumCombo.h"
 #include "ScopedTransaction.h"
@@ -15,6 +14,7 @@
 #include "LevelInstance/LevelInstanceInterface.h"
 #include "LevelInstance/LevelInstanceEditorPivotInterface.h"
 #include "LevelInstance/LevelInstanceEditorPivot.h"
+#include "SLevelInstancePivotPicker.h"
 
 #define LOCTEXT_NAMESPACE "FLevelInstancePivotDetails"
 
@@ -48,29 +48,23 @@ void FLevelInstancePivotDetails::CustomizeDetails(IDetailLayoutBuilder& DetailBu
 
 	PivotActor = nullptr;
 	PivotType = ELevelInstancePivotType::CenterMinZ;
-	PivotActors.Reset();
-	for (AActor* Actor : EditingObject->GetLevel()->Actors)
-	{
-		if (Actor != nullptr && Actor != EditingObject && Actor->GetRootComponent() != nullptr && !Actor->IsA<ABrush>())
-		{
-			PivotActors.Add(MakeShareable(new FActorInfo{ Actor }));
-		}
-	}
 
 	IDetailCategoryBuilder& ChangePivotCategory = DetailBuilder.EditCategory("Change Pivot", FText::GetEmpty(), ECategoryPriority::Important);
-		
+	
 	ChangePivotCategory.AddCustomRow(LOCTEXT("PivotTypeRow", "Type"), false)
 		.NameContent()
 		[
 			SNew(STextBlock)
 			.Text(LOCTEXT("PivotTypeText", "Type"))
+			.Font(FAppStyle::GetFontStyle(TEXT("PropertyWindow.NormalFont")))
 		]
 		.ValueContent()
 		[
 			SNew(SEnumComboBox, StaticEnum<ELevelInstancePivotType>())
-			.ContentPadding(FCoreStyle::Get().GetMargin("StandardDialog.ContentPadding"))
+			.ContentPadding(2.f)
 			.CurrentValue(this, &FLevelInstancePivotDetails::GetSelectedPivotType)
-			.OnEnumSelectionChanged(this, &FLevelInstancePivotDetails::OnSelectedPivotTypeChanged)
+			.OnEnumSelectionChanged(this, &FLevelInstancePivotDetails::OnSelectedPivotTypeChanged, EditingObject)
+			.Font(FAppStyle::GetFontStyle(TEXT("PropertyWindow.NormalFont")))
 		];
 
 	ChangePivotCategory.AddCustomRow(LOCTEXT("PivotActorRow", "Actor"), false)
@@ -78,19 +72,13 @@ void FLevelInstancePivotDetails::CustomizeDetails(IDetailLayoutBuilder& DetailBu
 		[
 			SNew(STextBlock)
 			.Text(LOCTEXT("PivotActorText", "Actor"))
+			.Font(FAppStyle::GetFontStyle(TEXT("PropertyWindow.NormalFont")))
 		]
 		.ValueContent()
 		[
-			SNew(SComboBox<TSharedPtr<FActorInfo>>)
-			.ContentPadding(FCoreStyle::Get().GetMargin("StandardDialog.ContentPadding"))
-			.OptionsSource(&PivotActors)
-			.OnGenerateWidget(this, &FLevelInstancePivotDetails::OnGeneratePivotActorWidget)
-			.OnSelectionChanged(this, &FLevelInstancePivotDetails::OnSelectedPivotActorChanged)
+			SNew(SLevelInstancePivotPicker)
 			.IsEnabled(this, &FLevelInstancePivotDetails::IsPivotActorSelectionEnabled)
-			[
-				SNew(STextBlock)
-				.Text(this, &FLevelInstancePivotDetails::GetSelectedPivotActorText)
-			]
+			.OnPivotActorPicked(this, &FLevelInstancePivotDetails::OnPivotActorPicked, EditingObject)
 		];
 			
 	ChangePivotCategory.AddCustomRow(LOCTEXT("ApplyRow", "Apply"), false)
@@ -101,9 +89,13 @@ void FLevelInstancePivotDetails::CustomizeDetails(IDetailLayoutBuilder& DetailBu
 			[
 				SNew(SButton)
 				.IsEnabled(this, &FLevelInstancePivotDetails::IsApplyButtonEnabled, EditingObject)
-				.Text(LOCTEXT("ChangePivotButton", "Apply"))
 				.HAlign(HAlign_Center)
 				.OnClicked(this, &FLevelInstancePivotDetails::OnApplyButtonClicked, EditingObject)
+				[
+					SNew(STextBlock)
+					.Font(FAppStyle::Get().GetFontStyle("PropertyWindow.NormalFont"))
+					.Text(LOCTEXT("ChangePivotButton", "Apply"))
+				]
 			]
 		];
 }
@@ -128,34 +120,36 @@ int32 FLevelInstancePivotDetails::GetSelectedPivotType() const
 	return (int32)PivotType;
 }
 
-void FLevelInstancePivotDetails::OnSelectedPivotTypeChanged(int32 NewValue, ESelectInfo::Type SelectionType)
+void FLevelInstancePivotDetails::OnSelectedPivotTypeChanged(int32 NewValue, ESelectInfo::Type SelectionType, TWeakObjectPtr<AActor> LevelInstancePivot)
 {
 	PivotType = (ELevelInstancePivotType)NewValue;
-}
 
-TSharedRef<SWidget> FLevelInstancePivotDetails::OnGeneratePivotActorWidget(TSharedPtr<FActorInfo> ActorInfo) const
-{
-	AActor* Actor = ActorInfo->Actor.Get();
-
-	// If a row wasn't generated just create the default one, a simple text block of the item's name.
-	return SNew(STextBlock).Text(Actor ? FText::FromString(Actor->GetActorLabel()) : LOCTEXT("null", "null"));
+	ShowPivotLocation(LevelInstancePivot);
 }
 
-FText FLevelInstancePivotDetails::GetSelectedPivotActorText() const
+bool FLevelInstancePivotDetails::IsPivotActorSelectionEnabled() const
 {
-	return PivotActor.IsValid() ? FText::FromString(PivotActor->GetActorLabel()) : LOCTEXT("none", "None");
+	return PivotType == ELevelInstancePivotType::Actor;
 }
 
-void FLevelInstancePivotDetails::OnSelectedPivotActorChanged(TSharedPtr<FActorInfo> NewValue, ESelectInfo::Type SelectionType)
+void FLevelInstancePivotDetails::OnPivotActorPicked(AActor* InPivotActor, TWeakObjectPtr<AActor> LevelInstancePivot)
 {
-	PivotActor = NewValue->Actor;
+	PivotActor = InPivotActor;
+
+	ShowPivotLocation(LevelInstancePivot);
 }
 
-bool FLevelInstancePivotDetails::IsPivotActorSelectionEnabled() const
+void FLevelInstancePivotDetails::ShowPivotLocation(const TWeakObjectPtr<AActor>& LevelInstancePivot)
 {
-	return PivotType == ELevelInstancePivotType::Actor;
-}
+	// Avoid showing invalid location
+	if (PivotType == ELevelInstancePivotType::Actor && PivotActor == nullptr)
+	{
+		return;
+	}
 
+	FVector PivotLocation = FLevelInstanceEditorPivotHelper::GetPivot(CastChecked<ILevelInstanceEditorPivotInterface>(LevelInstancePivot.Get()), PivotType, PivotActor.Get());
+	FLevelInstanceEditorPivotHelper::ShowPivotLocation(PivotLocation);
+}
 
 
 #undef LOCTEXT_NAMESPACE
\ No newline at end of file
diff --git a/Engine/Source/Editor/LevelInstanceEditor/Private/LevelInstancePivotDetails.h b/Engine/Source/Editor/LevelInstanceEditor/Private/LevelInstancePivotDetails.h
--- a/Engine/Source/Editor/LevelInstanceEditor/Private/LevelInstancePivotDetails.h
+++ b/Engine/Source/Editor/LevelInstanceEditor/Private/LevelInstancePivotDetails.h
@@ -28,21 +28,14 @@ private:
 	FReply OnApplyButtonClicked(TWeakObjectPtr<AActor> LevelInstancePivot);
 
 	int32 GetSelectedPivotType() const;
-	void OnSelectedPivotTypeChanged(int32 NewValue, ESelectInfo::Type SelectionType);
+	void OnSelectedPivotTypeChanged(int32 NewValue, ESelectInfo::Type SelectionType, TWeakObjectPtr<AActor> LevelInstancePivot);
 
-	// struct to hold weak ptr
-	struct FActorInfo
-	{
-		TWeakObjectPtr<AActor> Actor;
-	};
-
-	TSharedRef<SWidget> OnGeneratePivotActorWidget(TSharedPtr<FActorInfo> Actor) const;
-	FText GetSelectedPivotActorText() const;
-	void OnSelectedPivotActorChanged(TSharedPtr<FActorInfo> NewValue, ESelectInfo::Type SelectionType);
+	void OnPivotActorPicked(AActor* InPivotActor, TWeakObjectPtr<AActor> LevelInstancePivot);
 	bool IsPivotActorSelectionEnabled() const;
 
+	void ShowPivotLocation(const TWeakObjectPtr<AActor>& LevelInstancePivot);
+
 	// Settings for Apply button
 	ELevelInstancePivotType PivotType;
 	TWeakObjectPtr<AActor> PivotActor;
-	TArray<TSharedPtr<FActorInfo>> PivotActors;
 };
\ No newline at end of file
diff --git a/Engine/Source/Editor/LevelInstanceEditor/Private/SLevelInstancePivotPicker.cpp b/Engine/Source/Editor/LevelInstanceEditor/Private/SLevelInstancePivotPicker.cpp
new file mode 100644
--- /dev/null
+++ b/Engine/Source/Editor/LevelInstanceEditor/Private/SLevelInstancePivotPicker.cpp
@@ -0,0 +1,88 @@
+// Copyright Epic Games, Inc. All Rights Reserved.
+#include "SLevelInstancePivotPicker.h"
+#include "Widgets/Text/STextBlock.h"
+#include "Modules/ModuleManager.h"
+#include "SceneOutlinerPublicTypes.h"
+#include "SceneOutlinerModule.h"
+#include "ISceneOutlinerTreeItem.h"
+#include "ActorMode.h"
+#include "ActorPickingMode.h"
+#include "ActorTreeItem.h"
+
+#define LOCTEXT_NAMESPACE "SLevelInstancePivotPicker"
+
+void SLevelInstancePivotPicker::Construct(const FArguments& InArgs)
+{
+	OnPivotActorPickedEvent = InArgs._OnPivotActorPicked;
+
+	FOnSceneOutlinerItemPicked OnItemPicked = FOnSceneOutlinerItemPicked::CreateLambda([this](TSharedRef<ISceneOutlinerTreeItem> Item)
+	{
+		if (FActorTreeItem* ActorItem = Item->CastTo<FActorTreeItem>())
+		{
+			if (ActorItem->IsValid())
+			{
+				OnPivotActorPicked(ActorItem->Actor.Get());
+			}
+		}
+	});
+
+	FCreateSceneOutlinerMode ModeFactory = FCreateSceneOutlinerMode::CreateLambda([OnItemPicked](SSceneOutliner* Outliner)
+	{
+		FActorModeParams Params;
+		Params.SceneOutliner = Outliner;
+		Params.bHideComponents = true;
+		Params.bHideLevelInstanceHierarchy = false;
+		Params.bHideUnloadedActors = true;
+		Params.bHideEmptyFolders = true;
+		Params.bCanInteractWithSelectableActorsOnly = false;
+		return new FActorPickingMode(Params, OnItemPicked);
+	});
+
+	FSceneOutlinerInitializationOptions InitOptions;
+	{
+		InitOptions.bShowHeaderRow = false;
+		InitOptions.bFocusSearchBoxWhenOpened = true;
+		InitOptions.bShowTransient = true;
+		InitOptions.ModeFactory = ModeFactory;
+	}
+
+	// Actor selector to allow the user to choose a parent actor
+	FSceneOutlinerModule& SceneOutlinerModule = FModuleManager::LoadModuleChecked<FSceneOutlinerModule>("SceneOutliner");
+
+	this->ChildSlot
+	[
+		SAssignNew(PivotActorPicker, SComboButton)
+		.ContentPadding(2.f)
+		.HasDownArrow(false)
+		.ButtonContent()
+		[
+			SNew(STextBlock)
+			.Text(this, &SLevelInstancePivotPicker::GetSelectedPivotActorText)
+			.Font(FAppStyle::GetFontStyle(TEXT("PropertyWindow.NormalFont")))
+		]
+		.MenuContent()
+		[
+			SNew(SBox)
+			.MaxDesiredHeight(400.0f)
+			.WidthOverride(300.0f)
+			[
+				SceneOutlinerModule.CreateActorPicker(InitOptions, FOnActorPicked())
+			]
+		]
+	];
+}
+
+FText SLevelInstancePivotPicker::GetSelectedPivotActorText() const
+{
+	return PivotActor.IsValid() ? FText::FromString(PivotActor->GetActorLabel()) : LOCTEXT("none", "None");
+}
+
+void SLevelInstancePivotPicker::OnPivotActorPicked(AActor* PickedActor)
+{
+	PivotActor = PickedActor;
+	PivotActorPicker->SetIsOpen(false);
+			
+	OnPivotActorPickedEvent.ExecuteIfBound(PickedActor);
+}
+
+#undef LOCTEXT_NAMESPACE
\ No newline at end of file
diff --git a/Engine/Source/Editor/LevelInstanceEditor/Private/SLevelInstancePivotPicker.h b/Engine/Source/Editor/LevelInstanceEditor/Private/SLevelInstancePivotPicker.h
new file mode 100644
--- /dev/null
+++ b/Engine/Source/Editor/LevelInstanceEditor/Private/SLevelInstancePivotPicker.h
@@ -0,0 +1,30 @@
+// Copyright Epic Games, Inc. All Rights Reserved.
+
+#pragma once
+
+#include "Widgets/SCompoundWidget.h"
+#include "Widgets/Input/SComboButton.h"
+
+//////////////////////////////////////////////////////////////////////
+class AActor;
+
+DECLARE_DELEGATE_OneParam(FOnPivotActorPicked, AActor*)
+
+class SLevelInstancePivotPicker : public SCompoundWidget
+{
+public:
+	SLATE_BEGIN_ARGS(SLevelInstancePivotPicker) {}
+		SLATE_EVENT(FOnPivotActorPicked, OnPivotActorPicked)
+	SLATE_END_ARGS()
+
+public:
+	void Construct(const FArguments& InArgs);
+
+private:
+	FText GetSelectedPivotActorText() const;
+	void OnPivotActorPicked(AActor* PickedActor);
+
+	FOnPivotActorPicked OnPivotActorPickedEvent;
+	TWeakObjectPtr<AActor> PivotActor;
+	TSharedPtr<SComboButton> PivotActorPicker;
+};
\ No newline at end of file
diff --git a/Engine/Source/Editor/LevelInstanceEditor/Private/SNewLevelInstanceDialog.cpp b/Engine/Source/Editor/LevelInstanceEditor/Private/SNewLevelInstanceDialog.cpp
--- a/Engine/Source/Editor/LevelInstanceEditor/Private/SNewLevelInstanceDialog.cpp
+++ b/Engine/Source/Editor/LevelInstanceEditor/Private/SNewLevelInstanceDialog.cpp
@@ -45,6 +45,7 @@
 #include "Widgets/SWidget.h"
 #include "Widgets/SWindow.h"
 #include "Widgets/Text/STextBlock.h"
+#include "SLevelInstancePivotPicker.h"
 
 class IPropertyHandle;
 
@@ -187,18 +188,7 @@ FReply SNewLevelInstanceDialog::OnCancelClicked()
 	return FReply::Handled();
 }
 
-TSharedRef<SWidget> FNewLevelInstanceParamsDetails::OnGeneratePivotActorWidget(AActor* Actor) const
-{
-	// If a row wasn't generated just create the default one, a simple text block of the item's name.
-	return SNew(STextBlock).Text(Actor ? FText::FromString(Actor->GetActorLabel()) : LOCTEXT("null", "null")).Font(IDetailLayoutBuilder::GetDetailFont());
-}
-
-FText FNewLevelInstanceParamsDetails::GetSelectedPivotActorText() const
-{
-	return CreationParams->PivotActor ? FText::FromString(CreationParams->PivotActor->GetActorLabel()) : LOCTEXT("none", "None");
-}
-
-void FNewLevelInstanceParamsDetails::OnSelectedPivotActorChanged(AActor* NewValue, ESelectInfo::Type SelectionType)
+void FNewLevelInstanceParamsDetails::OnSelectedPivotActorChanged(AActor* NewValue)
 {
 	CreationParams->PivotActor = NewValue;
 }
@@ -255,16 +245,9 @@ void FNewLevelInstanceParamsDetails::CustomizeDetails(IDetailLayoutBuilder& Deta
 		.MinDesiredWidth(Row.ValueWidget.MinWidth)
 		.MaxDesiredWidth(Row.ValueWidget.MaxWidth)
 		[
-			SNew(SComboBox<AActor*>)
-			.OptionsSource(&PivotActors)
-			.OnGenerateWidget(this, &FNewLevelInstanceParamsDetails::OnGeneratePivotActorWidget)
-			.OnSelectionChanged(this, &FNewLevelInstanceParamsDetails::OnSelectedPivotActorChanged)
+			SNew(SLevelInstancePivotPicker)
+			.OnPivotActorPicked(this, &FNewLevelInstanceParamsDetails::OnSelectedPivotActorChanged)
 			.IsEnabled(this, &FNewLevelInstanceParamsDetails::IsPivotActorSelectionEnabled)
-			[
-				SNew(STextBlock)
-				.Text(this, &FNewLevelInstanceParamsDetails::GetSelectedPivotActorText)
-				.Font(DetailBuilder.GetDetailFont())
-			]
 		];
 }
 
diff --git a/Engine/Source/Editor/LevelInstanceEditor/Private/SNewLevelInstanceDialog.h b/Engine/Source/Editor/LevelInstanceEditor/Private/SNewLevelInstanceDialog.h
--- a/Engine/Source/Editor/LevelInstanceEditor/Private/SNewLevelInstanceDialog.h
+++ b/Engine/Source/Editor/LevelInstanceEditor/Private/SNewLevelInstanceDialog.h
@@ -78,9 +78,7 @@ public:
 
 	virtual void CustomizeDetails(IDetailLayoutBuilder& DetailBuilder) override;
 
-	TSharedRef<SWidget> OnGeneratePivotActorWidget(AActor* Actor) const;
-	FText GetSelectedPivotActorText() const;
-	void OnSelectedPivotActorChanged(AActor* NewValue, ESelectInfo::Type SelectionType);
+	void OnSelectedPivotActorChanged(AActor* NewValue);
 	bool IsPivotActorSelectionEnabled() const;
 
 protected:
diff --git a/Engine/Source/Editor/SceneOutliner/Private/ActorMode.cpp b/Engine/Source/Editor/SceneOutliner/Private/ActorMode.cpp
--- a/Engine/Source/Editor/SceneOutliner/Private/ActorMode.cpp
+++ b/Engine/Source/Editor/SceneOutliner/Private/ActorMode.cpp
@@ -112,6 +112,7 @@ FActorMode::FActorMode(const FActorModeParams& Params)
 	, bHideLevelInstanceHierarchy(Params.bHideLevelInstanceHierarchy)
 	, bHideUnloadedActors(Params.bHideUnloadedActors)
 	, bHideEmptyFolders(Params.bHideEmptyFolders)
+	, bCanInteractWithSelectableActorsOnly(Params.bCanInteractWithSelectableActorsOnly)
 {
 	SceneOutliner->AddFilter(MakeShared<FActorFilter>(FActorTreeItem::FFilterPredicate::CreateLambda([this](const AActor* Actor)
 	{
@@ -362,6 +363,37 @@ bool FActorMode::IsActorDisplayable(const SSceneOutliner* SceneOutliner, const A
 		FLevelUtils::IsLevelVisible(Actor->GetLevel());											// Only show Actors whose level is visible
 }
 
+bool FActorMode::CanInteract(const ISceneOutlinerTreeItem& Item) const
+{
+	if (bCanInteractWithSelectableActorsOnly)
+	{
+		AActor* FoundActor = nullptr;
+		if (const FActorTreeItem* ActorTreeItem = Item.CastTo<FActorTreeItem>())
+		{
+			FoundActor = ActorTreeItem->Actor.Get();
+		}
+		else if (const FComponentTreeItem* ComponentTreeItem = Item.CastTo<FComponentTreeItem>())
+		{
+			if (UActorComponent* Component = ComponentTreeItem->Component.Get())
+			{
+				FoundActor = Component->GetOwner();
+			}
+		}
+
+		if (FoundActor)
+		{
+			const bool bInSelected = true;
+			const bool bSelectEvenIfHidden = true;		// @todo outliner: Is this actually OK?
+			if (!GEditor->CanSelectActor(FoundActor, bInSelected, bSelectEvenIfHidden))
+			{
+				return false;
+			}
+		}
+	}
+	
+	return true;
+}
+
 bool FActorMode::IsActorLevelDisplayable(ULevel* InLevel)
 {
 	// Don't show level tree item for the persistent level
diff --git a/Engine/Source/Editor/SceneOutliner/Private/ActorTreeItem.cpp b/Engine/Source/Editor/SceneOutliner/Private/ActorTreeItem.cpp
--- a/Engine/Source/Editor/SceneOutliner/Private/ActorTreeItem.cpp
+++ b/Engine/Source/Editor/SceneOutliner/Private/ActorTreeItem.cpp
@@ -262,15 +262,7 @@ private:
 			// Highlight actors that are exclusive to PlayWorld
 			return FLinearColor(0.9f, 0.8f, 0.4f);
 		}
-
-		// also darken items that are non selectable in the active mode(s)
-		const bool bInSelected = true;
-		const bool bSelectEvenIfHidden = true;		// @todo outliner: Is this actually OK?
-		if (!GEditor->CanSelectActor(Actor, bInSelected, bSelectEvenIfHidden))
-		{
-			return FSceneOutlinerCommonLabelData::DarkColor;
-		}
-
+		
 		return FSlateColor::UseForeground();
 	}
 
@@ -346,14 +338,7 @@ bool FActorTreeItem::CanInteract() const
 		return false;
 	}
 
-	const bool bInSelected = true;
-	const bool bSelectEvenIfHidden = true;		// @todo outliner: Is this actually OK?
-	if (!GEditor->CanSelectActor(ActorPtr, bInSelected, bSelectEvenIfHidden))
-	{
-		return false;
-	}
-
-	return true;
+	return WeakSceneOutliner.Pin()->GetMode()->CanInteract(*this);
 }
 
 TSharedRef<SWidget> FActorTreeItem::GenerateLabelWidget(ISceneOutliner& Outliner, const STableRow<FSceneOutlinerTreeItemPtr>& InRow)
diff --git a/Engine/Source/Editor/SceneOutliner/Private/ComponentTreeItem.cpp b/Engine/Source/Editor/SceneOutliner/Private/ComponentTreeItem.cpp
--- a/Engine/Source/Editor/SceneOutliner/Private/ComponentTreeItem.cpp
+++ b/Engine/Source/Editor/SceneOutliner/Private/ComponentTreeItem.cpp
@@ -220,16 +220,7 @@ bool FComponentTreeItem::CanInteract() const
 		return false;
 	}
 
-	const bool bInSelected = true;
-	const bool bSelectEvenIfHidden = true;		// @todo outliner: Is this actually OK?
-
-	AActor* ParentActor = ComponentPtr->GetOwner(); 
-	if (!GEditor->CanSelectActor(ParentActor, bInSelected, bSelectEvenIfHidden)) 
-	{
-		return false;
-	}
-
-	return true;
+	return WeakSceneOutliner.Pin()->GetMode()->CanInteract(*this);
 }
 
 TSharedRef<SWidget> FComponentTreeItem::GenerateLabelWidget(ISceneOutliner& Outliner, const STableRow<FSceneOutlinerTreeItemPtr>& InRow)
diff --git a/Engine/Source/Editor/SceneOutliner/Private/SceneOutlinerModule.cpp b/Engine/Source/Editor/SceneOutliner/Private/SceneOutlinerModule.cpp
--- a/Engine/Source/Editor/SceneOutliner/Private/SceneOutlinerModule.cpp
+++ b/Engine/Source/Editor/SceneOutliner/Private/SceneOutlinerModule.cpp
@@ -78,32 +78,36 @@ TSharedRef<ISceneOutliner> FSceneOutlinerModule::CreateSceneOutliner(const FScen
 
 TSharedRef<ISceneOutliner> FSceneOutlinerModule::CreateActorPicker(const FSceneOutlinerInitializationOptions& InInitOptions, const FOnActorPicked& OnActorPickedDelegate, TWeakObjectPtr<UWorld> SpecifiedWorld) const
 {
-	auto OnItemPicked = FOnSceneOutlinerItemPicked::CreateLambda([OnActorPickedDelegate](TSharedRef<ISceneOutlinerTreeItem> Item)
-		{
-			if (FActorTreeItem* ActorItem = Item->CastTo<FActorTreeItem>())
+	FSceneOutlinerInitializationOptions InitOptions(InInitOptions);
+	if (!InitOptions.ModeFactory.IsBound())
+	{
+		auto OnItemPicked = FOnSceneOutlinerItemPicked::CreateLambda([OnActorPickedDelegate](TSharedRef<ISceneOutlinerTreeItem> Item)
 			{
-				if (ActorItem->IsValid())
+				if (FActorTreeItem* ActorItem = Item->CastTo<FActorTreeItem>())
 				{
-					OnActorPickedDelegate.ExecuteIfBound(ActorItem->Actor.Get());
+					if (ActorItem->IsValid())
+					{
+						OnActorPickedDelegate.ExecuteIfBound(ActorItem->Actor.Get());
 
+					}
 				}
-			}
-		});
+			});
 
-	FCreateSceneOutlinerMode ModeFactory = FCreateSceneOutlinerMode::CreateLambda([&OnItemPicked, &SpecifiedWorld](SSceneOutliner* Outliner)
-		{
-			FActorModeParams Params;
-			Params.SceneOutliner = Outliner;
-			Params.SpecifiedWorldToDisplay = SpecifiedWorld;
-			Params.bHideComponents = true;
-			Params.bHideLevelInstanceHierarchy = true;
-			Params.bHideUnloadedActors = true;
-			Params.bHideEmptyFolders = true;
-			return new FActorPickingMode(Params, OnItemPicked);
-		});
+		FCreateSceneOutlinerMode ModeFactory = FCreateSceneOutlinerMode::CreateLambda([OnItemPicked, SpecifiedWorld](SSceneOutliner* Outliner)
+			{
+				FActorModeParams Params;
+				Params.SceneOutliner = Outliner;
+				Params.SpecifiedWorldToDisplay = SpecifiedWorld;
+				Params.bHideComponents = true;
+				Params.bHideLevelInstanceHierarchy = true;
+				Params.bHideUnloadedActors = true;
+				Params.bHideEmptyFolders = true;
+				return new FActorPickingMode(Params, OnItemPicked);
+			});
+
+		InitOptions.ModeFactory = ModeFactory;
+	}
 
-	FSceneOutlinerInitializationOptions InitOptions(InInitOptions);
-	InitOptions.ModeFactory = ModeFactory;
 	if (InitOptions.ColumnMap.Num() == 0)
 	{
 		InitOptions.ColumnMap.Add(FSceneOutlinerBuiltInColumnTypes::Label(), FSceneOutlinerColumnInfo(ESceneOutlinerColumnVisibility::Visible, 0, FCreateSceneOutlinerColumn(), false, TOptional<float>(), FSceneOutlinerBuiltInColumnTypes::Label_Localized()));
diff --git a/Engine/Source/Editor/SceneOutliner/Public/ActorMode.h b/Engine/Source/Editor/SceneOutliner/Public/ActorMode.h
--- a/Engine/Source/Editor/SceneOutliner/Public/ActorMode.h
+++ b/Engine/Source/Editor/SceneOutliner/Public/ActorMode.h
@@ -29,13 +29,14 @@ struct SCENEOUTLINER_API FActorModeParams
 {
 	FActorModeParams() {}
 
-	FActorModeParams(SSceneOutliner* InSceneOutliner, const TWeakObjectPtr<UWorld>& InSpecifiedWorldToDisplay = nullptr, bool bInHideComponents = true, bool bInHideLevelInstanceHierarchy = true, bool bInHideUnloadedActors = true, bool bInHideEmptyFolders = true)
+	FActorModeParams(SSceneOutliner* InSceneOutliner, const TWeakObjectPtr<UWorld>& InSpecifiedWorldToDisplay = nullptr, bool bInHideComponents = true, bool bInHideLevelInstanceHierarchy = true, bool bInHideUnloadedActors = true, bool bInHideEmptyFolders = true, bool bInCanInteractWithSelectableActorsOnly = true)
 		: SpecifiedWorldToDisplay(InSpecifiedWorldToDisplay)
 		, SceneOutliner(InSceneOutliner)
 		, bHideComponents(bInHideComponents)
 		, bHideLevelInstanceHierarchy(bInHideLevelInstanceHierarchy)
 		, bHideUnloadedActors(bInHideUnloadedActors)
 		, bHideEmptyFolders(bInHideEmptyFolders)
+		, bCanInteractWithSelectableActorsOnly(bInCanInteractWithSelectableActorsOnly)
 	{}
 
 	TWeakObjectPtr<UWorld> SpecifiedWorldToDisplay = nullptr;
@@ -45,6 +46,7 @@ struct SCENEOUTLINER_API FActorModeParams
 	bool bHideLevelInstanceHierarchy = true;
 	bool bHideUnloadedActors = true;
 	bool bHideEmptyFolders = true;
+	bool bCanInteractWithSelectableActorsOnly = true;
 };
 
 class SCENEOUTLINER_API FActorMode : public ISceneOutlinerMode
@@ -74,6 +76,8 @@ public:
 	virtual FFolder::FRootObject GetRootObject() const override;
 	virtual FFolder::FRootObject GetPasteTargetRootObject() const override;
 
+	virtual bool CanInteract(const ISceneOutlinerTreeItem& Item) const override;
+
 private:
 	/** Called when the user selects a world in the world picker menu */
 	void OnSelectWorld(TWeakObjectPtr<UWorld> World);
@@ -111,4 +115,6 @@ protected:
 	bool bHideEmptyFolders;
 	/** Should the outliner scroll to the item on selection */
 	bool bAlwaysFrameSelection;
+	/** If True, CanInteract will be restricted to selectable actors only. */
+	bool bCanInteractWithSelectableActorsOnly;
 };
diff --git a/Engine/Source/Editor/SceneOutliner/Public/ActorPickingMode.h b/Engine/Source/Editor/SceneOutliner/Public/ActorPickingMode.h
--- a/Engine/Source/Editor/SceneOutliner/Public/ActorPickingMode.h
+++ b/Engine/Source/Editor/SceneOutliner/Public/ActorPickingMode.h
@@ -4,7 +4,7 @@
 
 #include "ActorMode.h"
 
-class FActorPickingMode : public FActorMode
+class SCENEOUTLINER_API FActorPickingMode : public FActorMode
 {
 public:
 	FActorPickingMode(const FActorModeParams& Params, FOnSceneOutlinerItemPicked OnItemPickedDelegate);
diff --git a/Engine/Source/Editor/SceneOutliner/Public/ISceneOutlinerMode.h b/Engine/Source/Editor/SceneOutliner/Public/ISceneOutlinerMode.h
--- a/Engine/Source/Editor/SceneOutliner/Public/ISceneOutlinerMode.h
+++ b/Engine/Source/Editor/SceneOutliner/Public/ISceneOutlinerMode.h
@@ -84,6 +84,8 @@ public:
 	virtual bool CanRenameItem(const ISceneOutlinerTreeItem& Item) const { return false; }
 	/** Whether the toolbar can be customized. */
 	virtual bool CanCustomizeToolbar() const { return false; }
+	/** Check if an item is interactive */
+	virtual bool CanInteract(const ISceneOutlinerTreeItem& Item) const { return true; }
 
 	/** Synchronize the mode specific selection with the tree view */
 	virtual void SynchronizeSelection() {}
diff --git a/Engine/Source/Editor/SceneOutliner/Public/ISceneOutlinerTreeItem.h b/Engine/Source/Editor/SceneOutliner/Public/ISceneOutlinerTreeItem.h
--- a/Engine/Source/Editor/SceneOutliner/Public/ISceneOutlinerTreeItem.h
+++ b/Engine/Source/Editor/SceneOutliner/Public/ISceneOutlinerTreeItem.h
@@ -7,6 +7,7 @@
 #include "Widgets/Views/STableRow.h"
 #include "SceneOutlinerStandaloneTypes.h"
 
+class ISceneOutliner;
 class UToolMenu;
 
 /** Base tree item interface  */
@@ -39,6 +40,9 @@ public:
 	/** Flags for this item */
 	FlagsType Flags;
 
+	/** Weak Outliner pointer */
+	TWeakPtr<ISceneOutliner> WeakSceneOutliner;
+
 	/** Delegate for hooking up an inline editable text block to be notified that a rename is requested. */
 	DECLARE_DELEGATE( FOnRenameRequest );
 
diff --git a/Engine/Source/Editor/SceneOutliner/Public/SSceneOutliner.h b/Engine/Source/Editor/SceneOutliner/Public/SSceneOutliner.h
--- a/Engine/Source/Editor/SceneOutliner/Public/SSceneOutliner.h
+++ b/Engine/Source/Editor/SceneOutliner/Public/SSceneOutliner.h
@@ -569,6 +569,7 @@ private:
 		if (bForce || bPassesFilters)
 		{
 			FSceneOutlinerTreeItemPtr Result = MakeShareable(new TreeItemType(Data));
+			Result->WeakSceneOutliner = StaticCastSharedRef<ISceneOutliner>(AsShared());
 			Result->Flags.bIsFilteredOut = !bPassesFilters;
 			Result->Flags.bInteractive = Filters->GetInteractiveState(*Result) && InteractiveFilters->GetInteractiveState(*Result);
 			return Result;
diff --git a/Engine/Source/Runtime/Engine/Private/LevelInstance/LevelInstanceEditorPivot.cpp b/Engine/Source/Runtime/Engine/Private/LevelInstance/LevelInstanceEditorPivot.cpp
--- a/Engine/Source/Runtime/Engine/Private/LevelInstance/LevelInstanceEditorPivot.cpp
+++ b/Engine/Source/Runtime/Engine/Private/LevelInstance/LevelInstanceEditorPivot.cpp
@@ -12,6 +12,8 @@
 #include "GameFramework/WorldSettings.h"
 #include "LevelUtils.h"
 #include "Editor.h"
+#include "UnrealEdGlobals.h"
+#include "Editor/UnrealEdEngine.h"
 
 #endif
 
@@ -72,17 +74,40 @@ ILevelInstanceEditorPivotInterface* FLevelInstanceEditorPivotHelper::Create(ILev
 	return PivotInterface;
 }
 
+void FLevelInstanceEditorPivotHelper::ShowPivotLocation(const FVector& PivotLocation)
+{
+	if (GEditor)
+	{
+		// Show in viewport where the Pivot would be located
+		GEditor->SetPivot(PivotLocation, false, true);
+		GUnrealEd->SetPivotMovedIndependently(true);
+		GEditor->RedrawLevelEditingViewports(false);
+	}
+}
+
 void FLevelInstanceEditorPivotHelper::SetPivot(ILevelInstanceEditorPivotInterface* PivotInterface, ELevelInstancePivotType PivotType, AActor* PivotToActor)
+{
+	FVector Location = GetPivot(PivotInterface, PivotType, PivotToActor);
+	AActor* PivotActor = CastChecked<AActor>(PivotInterface);
+		
+	PivotActor->Modify();
+	PivotActor->SetActorLocation(Location);
+
+	PivotInterface->UpdateOffset();
+
+	ShowPivotLocation(Location);
+}
+
+FVector FLevelInstanceEditorPivotHelper::GetPivot(ILevelInstanceEditorPivotInterface* PivotInterface, ELevelInstancePivotType PivotType, AActor* PivotToActor)
 {
 	check(PivotType != ELevelInstancePivotType::Actor || PivotToActor != nullptr);
 	AActor* PivotActor = CastChecked<AActor>(PivotInterface);
 
-	PivotActor->Modify();
 	if (PivotType == ELevelInstancePivotType::Actor)
 	{
-		PivotActor->SetActorLocation(PivotToActor->GetActorLocation());
+		return PivotToActor->GetActorLocation();
 	}
-	else if(PivotType == ELevelInstancePivotType::Center || PivotType == ELevelInstancePivotType::CenterMinZ)
+	else if (PivotType == ELevelInstancePivotType::Center || PivotType == ELevelInstancePivotType::CenterMinZ)
 	{
 		ULevelInstanceSubsystem* LevelInstanceSubsystem = UWorld::GetSubsystem<ULevelInstanceSubsystem>(PivotActor->GetWorld());
 		ILevelInstanceInterface* LevelInstance = LevelInstanceSubsystem->GetLevelInstance(PivotInterface->GetLevelInstanceID());
@@ -93,24 +118,13 @@ void FLevelInstanceEditorPivotHelper::SetPivot(ILevelInstanceEditorPivotInterfac
 		{
 			Location.Z = OutBounds.Min.Z;
 		}
-		PivotActor->SetActorLocation(Location);
+		return Location;
 	}
-	else if (PivotType == ELevelInstancePivotType::WorldOrigin)
+	else
 	{
-		PivotActor->SetActorLocation(FVector(0.f, 0.f, 0.f));
+		check(PivotType == ELevelInstancePivotType::WorldOrigin);
+		return FVector::ZeroVector;
 	}
-	else // unsupported
-	{
-		check(0);
-	}
-
-	// Update gizmo
-	if (GEditor)
-	{
-		GEditor->NoteSelectionChange();
-	}
-
-	PivotInterface->UpdateOffset();
 }
 
 #endif
diff --git a/Engine/Source/Runtime/Engine/Public/LevelInstance/LevelInstanceEditorPivot.h b/Engine/Source/Runtime/Engine/Public/LevelInstance/LevelInstanceEditorPivot.h
--- a/Engine/Source/Runtime/Engine/Public/LevelInstance/LevelInstanceEditorPivot.h
+++ b/Engine/Source/Runtime/Engine/Public/LevelInstance/LevelInstanceEditorPivot.h
@@ -16,7 +16,9 @@ class ENGINE_API FLevelInstanceEditorPivotHelper
 {
 public:
 	static ILevelInstanceEditorPivotInterface* Create(ILevelInstanceInterface* LevelInstance, ULevelStreaming* LevelStreaming);
-	static void SetPivot(ILevelInstanceEditorPivotInterface* Pivot, ELevelInstancePivotType PivotType, AActor* PivotToActor);
+	static void SetPivot(ILevelInstanceEditorPivotInterface* PivotInterface, ELevelInstancePivotType PivotType, AActor* PivotToActor);
+	static FVector GetPivot(ILevelInstanceEditorPivotInterface* PivotInterface, ELevelInstancePivotType PivotType, AActor* PivotToActor);
+	static void ShowPivotLocation(const FVector& PivotLocation);
 };
 
 #endif
\ No newline at end of file
