LevelInstance: Blueprint Filter editing support

#rnx

diff --git a/Engine/Source/Editor/LevelInstanceEditor/LevelInstanceEditor.Build.cs b/Engine/Source/Editor/LevelInstanceEditor/LevelInstanceEditor.Build.cs
--- a/Engine/Source/Editor/LevelInstanceEditor/LevelInstanceEditor.Build.cs
+++ b/Engine/Source/Editor/LevelInstanceEditor/LevelInstanceEditor.Build.cs
@@ -33,7 +33,8 @@ public class LevelInstanceEditor : ModuleRules
 				"EditorWidgets",
 				"DeveloperSettings",
 				"SceneOutliner",
-				"WorldPartitionEditor"
+				"WorldPartitionEditor",
+				"Kismet"
 			}
 		);
 
diff --git a/Engine/Source/Editor/LevelInstanceEditor/Private/LevelInstanceFilterPropertyTypeCustomization.cpp b/Engine/Source/Editor/LevelInstanceEditor/Private/LevelInstanceFilterPropertyTypeCustomization.cpp
--- a/Engine/Source/Editor/LevelInstanceEditor/Private/LevelInstanceFilterPropertyTypeCustomization.cpp
+++ b/Engine/Source/Editor/LevelInstanceEditor/Private/LevelInstanceFilterPropertyTypeCustomization.cpp
@@ -8,6 +8,11 @@
 #include "Engine/World.h"
 #include "GameFramework/Actor.h"
 #include "ScopedTransaction.h"
+#include "Editor.h"
+#include "BlueprintEditor.h"
+#include "PreviewScene.h"
+#include "PropertyHandle.h"
+#include "Subsystems/AssetEditorSubsystem.h"
 
 #define LOCTEXT_NAMESPACE "WorldPartitionActorFilter"
 
@@ -30,8 +35,22 @@ TSharedPtr<FWorldPartitionActorFilterMode::FFilter> FLevelInstanceFilterProperty
 
 				if (WorldAssetPackage.IsEmpty())
 				{
-					World = OuterActor->GetWorld();
 					WorldAssetPackage = LevelInstanceInterface->GetWorldAssetPackage();
+
+					World = OuterActor->GetWorld();
+
+					// Handle Blueprint Editing. In this case we are going to rely on the preview world to build the filters
+					if (!World && OuterActor->IsTemplate() && OuterActor->GetClass()->ClassGeneratedBy)
+					{
+						if (FBlueprintEditor* BlueprintEditor = static_cast<FBlueprintEditor*>(GEditor->GetEditorSubsystem<UAssetEditorSubsystem>()->FindEditorForAsset(OuterActor->GetClass()->ClassGeneratedBy, false)))
+						{
+							if (FPreviewScene* PreviewScene = BlueprintEditor->GetPreviewScene())
+							{
+								World = PreviewScene->GetWorld();
+								break;
+							}
+						}
+					}
 				}
 				else if (WorldAssetPackage != LevelInstanceInterface->GetWorldAssetPackage())
 				{
@@ -65,7 +84,7 @@ TSharedPtr<FWorldPartitionActorFilterMode::FFilter> FLevelInstanceFilterProperty
 	return MakeShared<FWorldPartitionActorFilterMode::FFilter>(Filter, SelectedFilters);
 }
 
-void FLevelInstanceFilterPropertyTypeCustomization::ApplyFilter(const FWorldPartitionActorFilterMode& Mode)
+void FLevelInstanceFilterPropertyTypeCustomization::ApplyFilter(TSharedRef<IPropertyHandle> PropertyHandle, const FWorldPartitionActorFilterMode& Mode)
 {
 	const FScopedTransaction Transaction(LOCTEXT("WorldPartitionActorFilterApply_Transaction", "Apply Level Instance Filter"));
 	for (ILevelInstanceInterface* LevelInstance : LevelInstances)
@@ -73,7 +92,9 @@ void FLevelInstanceFilterPropertyTypeCustomization::ApplyFilter(const FWorldPart
 		ULevelInstanceComponent* Component = LevelInstance->GetLevelInstanceComponent();
 		FWorldPartitionActorFilter ComponentFilter = Component->GetFilter();
 		Mode.Apply(ComponentFilter);
-		Component->SetFilter(ComponentFilter);
+		
+		// Setting value through PropertyHandle will take care of propagating to instances if we are editing a template (Blueprint)
+		PropertyHandle->SetValueFromFormattedString(ComponentFilter.ToString(), EPropertyValueSetFlags::DefaultFlags);
 	}
 }
 
diff --git a/Engine/Source/Editor/LevelInstanceEditor/Private/LevelInstanceFilterPropertyTypeCustomization.h b/Engine/Source/Editor/LevelInstanceEditor/Private/LevelInstanceFilterPropertyTypeCustomization.h
--- a/Engine/Source/Editor/LevelInstanceEditor/Private/LevelInstanceFilterPropertyTypeCustomization.h
+++ b/Engine/Source/Editor/LevelInstanceEditor/Private/LevelInstanceFilterPropertyTypeCustomization.h
@@ -30,7 +30,7 @@ public:
 	}
 private:
 	virtual TSharedPtr<FWorldPartitionActorFilterMode::FFilter> CreateModeFilter(TArray<UObject*> OuterObjects) override;
-	virtual void ApplyFilter(const FWorldPartitionActorFilterMode& Mode) override;
+	virtual void ApplyFilter(TSharedRef<IPropertyHandle> PropertyHandle, const FWorldPartitionActorFilterMode& Mode) override;
 
 	TArray<ILevelInstanceInterface*> LevelInstances;
 };
diff --git a/Engine/Source/Editor/WorldPartitionEditor/Private/WorldPartition/Filter/WorldPartitionActorFilterPropertyTypeCustomization.cpp b/Engine/Source/Editor/WorldPartitionEditor/Private/WorldPartition/Filter/WorldPartitionActorFilterPropertyTypeCustomization.cpp
--- a/Engine/Source/Editor/WorldPartitionEditor/Private/WorldPartition/Filter/WorldPartitionActorFilterPropertyTypeCustomization.cpp
+++ b/Engine/Source/Editor/WorldPartitionEditor/Private/WorldPartition/Filter/WorldPartitionActorFilterPropertyTypeCustomization.cpp
@@ -81,10 +81,10 @@ void FWorldPartitionActorFilterPropertyTypeCustomization::CustomizeHeader(TShare
 
 	SceneOutliner = SceneOutlinerModule.CreateSceneOutliner(InitOptions);
 
-	auto OnApplyFilter = [OnValueChanged, this]()
+	auto OnApplyFilter = [StructPropertyHandle, OnValueChanged, this]()
 	{
 		const FWorldPartitionActorFilterMode* Mode = static_cast<const FWorldPartitionActorFilterMode*>(SceneOutliner->GetMode());
-		ApplyFilter(*Mode);
+		ApplyFilter(StructPropertyHandle, *Mode);
 		OnValueChanged.Execute();
 		return FReply::Handled();
 	};
diff --git a/Engine/Source/Editor/WorldPartitionEditor/Public/WorldPartition/Filter/WorldPartitionActorFilterPropertyTypeCustomization.h b/Engine/Source/Editor/WorldPartitionEditor/Public/WorldPartition/Filter/WorldPartitionActorFilterPropertyTypeCustomization.h
--- a/Engine/Source/Editor/WorldPartitionEditor/Public/WorldPartition/Filter/WorldPartitionActorFilterPropertyTypeCustomization.h
+++ b/Engine/Source/Editor/WorldPartitionEditor/Public/WorldPartition/Filter/WorldPartitionActorFilterPropertyTypeCustomization.h
@@ -24,7 +24,7 @@ public:
 
 private:
 	virtual TSharedPtr<FWorldPartitionActorFilterMode::FFilter> CreateModeFilter(TArray<UObject*> OuterObjects) = 0;
-	virtual void ApplyFilter(const FWorldPartitionActorFilterMode& Mode) = 0;
+	virtual void ApplyFilter(TSharedRef<IPropertyHandle> PropertyHandle, const FWorldPartitionActorFilterMode& Mode) = 0;
 
 	FDelegateHandle WorldPartitionActorFilterChangedHandle;
 	TSharedPtr<ISceneOutliner> SceneOutliner;
diff --git a/Engine/Source/Runtime/Engine/Private/WorldPartition/DataLayer/DataLayerSubsystem.cpp b/Engine/Source/Runtime/Engine/Private/WorldPartition/DataLayer/DataLayerSubsystem.cpp
--- a/Engine/Source/Runtime/Engine/Private/WorldPartition/DataLayer/DataLayerSubsystem.cpp
+++ b/Engine/Source/Runtime/Engine/Private/WorldPartition/DataLayer/DataLayerSubsystem.cpp
@@ -30,7 +30,7 @@
 
 bool UDataLayerSubsystem::DoesSupportWorldType(const EWorldType::Type WorldType) const
 {
-	return Super::DoesSupportWorldType(WorldType) || WorldType == EWorldType::Inactive;
+	return Super::DoesSupportWorldType(WorldType) || WorldType == EWorldType::Inactive || WorldType == EWorldType::EditorPreview;
 }
 
 #if WITH_EDITOR
diff --git a/Engine/Source/Runtime/Engine/Private/WorldPartition/Filter/WorldPartitionActorFilter.cpp b/Engine/Source/Runtime/Engine/Private/WorldPartition/Filter/WorldPartitionActorFilter.cpp
--- a/Engine/Source/Runtime/Engine/Private/WorldPartition/Filter/WorldPartitionActorFilter.cpp
+++ b/Engine/Source/Runtime/Engine/Private/WorldPartition/Filter/WorldPartitionActorFilter.cpp
@@ -225,6 +225,18 @@ FArchive& operator<<(FArchive& Ar, FWorldPartitionActorFilter& Filter)
 	return Ar;
 }
 
+FString FWorldPartitionActorFilter::ToString() const
+{
+	FString ReturnValue;
+	FWorldPartitionActorFilter Default;
+	if (ExportTextItem(ReturnValue, Default, nullptr, 0, nullptr))
+	{
+		return ReturnValue;
+	}
+
+	return FString();
+}
+
 bool FWorldPartitionActorFilter::ExportTextItem(FString& ValueStr, FWorldPartitionActorFilter const& DefaultValue, UObject* ParentObject, int32 PortFlags, UObject* ExportRootScope) const
 {
 	const uint32 OriginalLen = ValueStr.Len();
diff --git a/Engine/Source/Runtime/Engine/Private/WorldPartition/WorldPartitionSubsystem.cpp b/Engine/Source/Runtime/Engine/Private/WorldPartition/WorldPartitionSubsystem.cpp
--- a/Engine/Source/Runtime/Engine/Private/WorldPartition/WorldPartitionSubsystem.cpp
+++ b/Engine/Source/Runtime/Engine/Private/WorldPartition/WorldPartitionSubsystem.cpp
@@ -167,7 +167,7 @@ void UWorldPartitionSubsystem::Deinitialize()
 // because DataLayerManager is required to be initialized when duplicating a partitioned world.
 bool UWorldPartitionSubsystem::DoesSupportWorldType(const EWorldType::Type WorldType) const
 {
-	return Super::DoesSupportWorldType(WorldType) || WorldType == EWorldType::Inactive;
+	return Super::DoesSupportWorldType(WorldType) || WorldType == EWorldType::Inactive || WorldType == EWorldType::EditorPreview;
 }
 
 void UWorldPartitionSubsystem::ForEachWorldPartition(TFunctionRef<bool(UWorldPartition*)> Func)
diff --git a/Engine/Source/Runtime/Engine/Public/WorldPartition/Filter/WorldPartitionActorFilter.h b/Engine/Source/Runtime/Engine/Public/WorldPartition/Filter/WorldPartitionActorFilter.h
--- a/Engine/Source/Runtime/Engine/Public/WorldPartition/Filter/WorldPartitionActorFilter.h
+++ b/Engine/Source/Runtime/Engine/Public/WorldPartition/Filter/WorldPartitionActorFilter.h
@@ -70,6 +70,7 @@ struct ENGINE_API FWorldPartitionActorFilter
 	// Needed for Copy/Paste/ResetToDefault
 	bool ExportTextItem(FString& ValueStr, FWorldPartitionActorFilter const& DefaultValue, UObject* Parent, int32 PortFlags, UObject* ExportRootScope) const;
 	bool ImportTextItem(const TCHAR*& Buffer, int32 PortFlags, UObject* Parent, FOutputDevice* ErrorText);
+	FString ToString() const;
 
 	static void RequestFilterRefresh(bool bIsFromUserChange);
 
