WorldPartition: UnloadReason now provided by FActorDescFilter

#rnx

diff --git a/Engine/Source/Editor/DataLayerEditor/Private/DataLayer/DataLayerEditorSubsystem.cpp b/Engine/Source/Editor/DataLayerEditor/Private/DataLayer/DataLayerEditorSubsystem.cpp
--- a/Engine/Source/Editor/DataLayerEditor/Private/DataLayer/DataLayerEditorSubsystem.cpp
+++ b/Engine/Source/Editor/DataLayerEditor/Private/DataLayer/DataLayerEditorSubsystem.cpp
@@ -200,8 +200,34 @@ void UDataLayerEditorSubsystem::Initialize(FSubsystemCollectionBase& Collection)
 	// Register the engine broadcast bridge
 	OnActorDataLayersEditorLoadingStateChangedEngineBridgeHandle = DataLayerEditorLoadingStateChanged.AddStatic(&FDataLayersEditorBroadcast::StaticOnActorDataLayersEditorLoadingStateChanged);
 
+	class FDataLayerActorDescFilter : public IWorldPartitionActorLoaderInterface::FActorDescFilter
+	{
+	public:
+		FDataLayerActorDescFilter(UDataLayerEditorSubsystem* InSubsystem) : Subsystem(InSubsystem) {}
+		bool PassFilter(class UWorld* InWorld, const FWorldPartitionHandle& InHandle) override
+		{
+			if (!Subsystem->PassDataLayersFilter(InWorld, InHandle))
+			{
+				return false;
+			}
+
+			return true;
+		}
+
+		// Leave [0,9] for Game code
+		virtual uint32 GetFilterPriority() const { return 10; }
+
+		virtual FText* GetFilterReason() const override
+		{
+			static FText UnloadedReason(LOCTEXT("DataLayerFilterReason", "Unloaded Datalayer"));
+			return &UnloadedReason;
+		}
+	private:
+		UDataLayerEditorSubsystem* Subsystem;
+	};
+
 	// Register actor descriptor loading filter
-	IWorldPartitionActorLoaderInterface::RegisterActorDescFilter([this](UWorld* World, const FWorldPartitionHandle& ActorHandle) { return PassDataLayersFilter(World, ActorHandle); });
+	IWorldPartitionActorLoaderInterface::RegisterActorDescFilter(MakeShareable<IWorldPartitionActorLoaderInterface::FActorDescFilter>(new FDataLayerActorDescFilter(this)));
 }
 
 void UDataLayerEditorSubsystem::Deinitialize()
diff --git a/Engine/Source/Editor/LevelInstanceEditor/Private/LevelInstanceEditorModule.cpp b/Engine/Source/Editor/LevelInstanceEditor/Private/LevelInstanceEditorModule.cpp
--- a/Engine/Source/Editor/LevelInstanceEditor/Private/LevelInstanceEditorModule.cpp
+++ b/Engine/Source/Editor/LevelInstanceEditor/Private/LevelInstanceEditorModule.cpp
@@ -682,18 +682,32 @@ void FLevelInstanceEditorModule::StartupModule()
 	EditorLevelUtils::CanMoveActorToLevelDelegate.AddRaw(this, &FLevelInstanceEditorModule::CanMoveActorToLevel);
 
 	// Register actor descriptor loading filter
-	IWorldPartitionActorLoaderInterface::RegisterActorDescFilter([](UWorld* World, const FWorldPartitionHandle& ActorHandle) 
-	{ 
-		if (UWorld* OwningWorld = World->PersistentLevel->GetWorld())
+	class FLevelInstanceActorDescFilter : public IWorldPartitionActorLoaderInterface::FActorDescFilter
+	{
+	public:
+		bool PassFilter(class UWorld* InWorld, const FWorldPartitionHandle& InHandle) override
 		{
-			if (ULevelInstanceSubsystem* LevelInstanceSubsystem = OwningWorld->GetSubsystem<ULevelInstanceSubsystem>())
+			if (UWorld* OwningWorld = InWorld->PersistentLevel->GetWorld())
 			{
-				return LevelInstanceSubsystem->PassLevelInstanceFilter(World, ActorHandle);
+				if (ULevelInstanceSubsystem* LevelInstanceSubsystem = OwningWorld->GetSubsystem<ULevelInstanceSubsystem>())
+				{
+					return LevelInstanceSubsystem->PassLevelInstanceFilter(InWorld, InHandle);
+				}
 			}
+
+			return true;
 		}
 
-		return true;
-	});
+		// Leave [0, 19] for Game code
+		virtual uint32 GetFilterPriority() const override { return 20; }
+
+		virtual FText* GetFilterReason() const override
+		{
+			static FText UnloadedReason(LOCTEXT("LevelInstanceActorDescFilterReason", "Filtered"));
+			return &UnloadedReason;
+		}
+	};
+	IWorldPartitionActorLoaderInterface::RegisterActorDescFilter(MakeShareable<IWorldPartitionActorLoaderInterface::FActorDescFilter>(new FLevelInstanceActorDescFilter()));
 
 	FMessageLogModule& MessageLogModule = FModuleManager::LoadModuleChecked<FMessageLogModule>("MessageLog");
 	FMessageLogInitializationOptions InitOptions;
diff --git a/Engine/Source/Editor/SceneOutliner/Private/ActorDescTreeItem.cpp b/Engine/Source/Editor/SceneOutliner/Private/ActorDescTreeItem.cpp
--- a/Engine/Source/Editor/SceneOutliner/Private/ActorDescTreeItem.cpp
+++ b/Engine/Source/Editor/SceneOutliner/Private/ActorDescTreeItem.cpp
@@ -125,7 +125,7 @@ private:
 			{
 				FFormatNamedArguments Args;
 				Args.Add(TEXT("ActorLabel"), FText::FromString(TreeItem->GetDisplayString()));
-				Args.Add(TEXT("UnloadState"), FLoaderAdapterPinnedActors::GetUnloadedReason(ActorDesc));
+				Args.Add(TEXT("UnloadState"), ActorDesc->GetUnloadedReason());
 
 				return FText::Format(LOCTEXT("UnloadedActorDisplay", "{ActorLabel} ({UnloadState})"), Args);
 			}
diff --git a/Engine/Source/Runtime/Engine/Private/WorldPartition/LoaderAdapter/LoaderAdapterPinnedActors.cpp b/Engine/Source/Runtime/Engine/Private/WorldPartition/LoaderAdapter/LoaderAdapterPinnedActors.cpp
--- a/Engine/Source/Runtime/Engine/Private/WorldPartition/LoaderAdapter/LoaderAdapterPinnedActors.cpp
+++ b/Engine/Source/Runtime/Engine/Private/WorldPartition/LoaderAdapter/LoaderAdapterPinnedActors.cpp
@@ -16,27 +16,6 @@ bool FLoaderAdapterPinnedActors::PassActorDescFilter(const FWorldPartitionHandle
 	return ActorHandle.IsValid() && !ActorsToRemove.Contains(ActorHandle);
 }
 
-FText FLoaderAdapterPinnedActors::GetUnloadedReason(FWorldPartitionActorDesc* InActorDesc)
-{
-	if (InActorDesc)
-	{
-		if (InActorDesc->FailedToLoad())
-		{
-			return LOCTEXT("FailedToLoadReason", "Failed to load");
-		}
-
-		UActorDescContainer* ActorDescContainer = InActorDesc->GetContainer();
-		UWorldPartition* WorldPartition = ActorDescContainer != nullptr ? ActorDescContainer->GetWorldPartition() : nullptr;
-		bool bShouldBeLoaded = !InActorDesc->GetIsSpatiallyLoaded() && InActorDesc->IsEditorRelevant();
-		if (WorldPartition && (bShouldBeLoaded || WorldPartition->IsActorPinned(InActorDesc->GetGuid())))
-		{
-			return LOCTEXT("UnloadedDataLayerReason", "Unloaded DataLayer");
-		}
-	}
-	
-	return LOCTEXT("UnloadedReason", "Unloaded");
-}
-
 bool FLoaderAdapterPinnedActors::SupportsPinning(FWorldPartitionActorDesc* InActorDesc)
 {
 	if (!InActorDesc)
diff --git a/Engine/Source/Runtime/Engine/Private/WorldPartition/WorldPartitionActorDesc.cpp b/Engine/Source/Runtime/Engine/Private/WorldPartition/WorldPartitionActorDesc.cpp
--- a/Engine/Source/Runtime/Engine/Private/WorldPartition/WorldPartitionActorDesc.cpp
+++ b/Engine/Source/Runtime/Engine/Private/WorldPartition/WorldPartitionActorDesc.cpp
@@ -27,6 +27,8 @@
 #include "WorldPartition/ErrorHandling/WorldPartitionStreamingGenerationErrorHandler.h"
 #include "ActorReferencesUtils.h"
 
+#define LOCTEXT_NAMESPACE "FWorldPartitionActorDesc"
+
 TMap<TSubclassOf<AActor>, FWorldPartitionActorDesc::FActorDescDeprecator> FWorldPartitionActorDesc::Deprecators;
 
 FWorldPartitionActorDesc::FWorldPartitionActorDesc()
@@ -35,7 +37,7 @@ FWorldPartitionActorDesc::FWorldPartitionActorDesc()
 	, HardRefCount(0)
 	, Container(nullptr)
 	, bIsForcedNonSpatiallyLoaded(false)
-	, bFailedToLoad(false)
+	, UnloadedReason(nullptr)
 {}
 
 void FWorldPartitionActorDesc::Init(const AActor* InActor)
@@ -641,9 +643,16 @@ AActor* FWorldPartitionActorDesc::GetActor(bool bEvenIfPendingKill, bool bEvenIf
 	return bEvenIfUnreachable ? ActorPtr.GetEvenIfUnreachable() : ActorPtr.Get(bEvenIfPendingKill);
 }
 
+const FText& FWorldPartitionActorDesc::GetUnloadedReason() const
+{
+	static FText Unloaded(LOCTEXT("UnloadedReason", "Unloaded"));
+	return UnloadedReason ? *UnloadedReason : Unloaded;
+}
+
 AActor* FWorldPartitionActorDesc::Load() const
 {
-	bFailedToLoad = false;
+	static FText FailedToLoad(LOCTEXT("FailedToLoadReason", "Failed to load"));
+	UnloadedReason = nullptr;
 
 	if (ActorPtr.IsExplicitlyNull() || ActorPtr.IsStale())
 	{
@@ -674,7 +683,7 @@ AActor* FWorldPartitionActorDesc::Load() const
 			if (!ActorPtr.IsValid())
 			{
 				UE_LOG(LogWorldPartition, Warning, TEXT("Can't load actor guid `%s` ('%s') from package '%s'"), *Guid.ToString(), *GetActorName().ToString(), *ActorPackage.ToString());
-				bFailedToLoad = true;
+				UnloadedReason = &FailedToLoad;
 			}
 		}
 	}
@@ -711,4 +720,7 @@ void FWorldPartitionActorDesc::Unload()
 		ActorPtr = nullptr;
 	}
 }
+
+#undef LOCTEXT_NAMESPACE
+
 #endif
diff --git a/Engine/Source/Runtime/Engine/Private/WorldPartition/WorldPartitionActorLoaderInterface.cpp b/Engine/Source/Runtime/Engine/Private/WorldPartition/WorldPartitionActorLoaderInterface.cpp
--- a/Engine/Source/Runtime/Engine/Private/WorldPartition/WorldPartitionActorLoaderInterface.cpp
+++ b/Engine/Source/Runtime/Engine/Private/WorldPartition/WorldPartitionActorLoaderInterface.cpp
@@ -18,12 +18,13 @@ UWorldPartitionActorLoaderInterface::UWorldPartitionActorLoaderInterface(const F
 {}
 
 #if WITH_EDITOR
-TArray<IWorldPartitionActorLoaderInterface::FActorDescFilter> IWorldPartitionActorLoaderInterface::ActorDescFilters;
+TArray<TSharedRef<IWorldPartitionActorLoaderInterface::FActorDescFilter>> IWorldPartitionActorLoaderInterface::ActorDescFilters;
 IWorldPartitionActorLoaderInterface::FOnActorLoaderInterfaceRefreshState IWorldPartitionActorLoaderInterface::ActorLoaderInterfaceRefreshState;
 
-void IWorldPartitionActorLoaderInterface::RegisterActorDescFilter(const FActorDescFilter& InActorDescFilter)
+void IWorldPartitionActorLoaderInterface::RegisterActorDescFilter(const TSharedRef<FActorDescFilter>& InActorDescFilter)
 {
 	ActorDescFilters.Add(InActorDescFilter);
+	ActorDescFilters.Sort([](const TSharedRef<FActorDescFilter>& FilterA, const TSharedRef<FActorDescFilter>& FilterB) { return FilterA.Get().GetFilterPriority() > FilterB.Get().GetFilterPriority(); });
 }
 
 IWorldPartitionActorLoaderInterface::ILoaderAdapter::ILoaderAdapter(UWorld* InWorld)
@@ -231,6 +232,8 @@ bool IWorldPartitionActorLoaderInterface::ILoaderAdapter::ShouldActorBeLoaded(co
 {
 	check(Actor.IsValid());
 
+	Actor->SetUnloadedReason(nullptr);
+
 	if (!PassActorDescFilter(Actor))
 	{
 		return false;
@@ -238,8 +241,9 @@ bool IWorldPartitionActorLoaderInterface::ILoaderAdapter::ShouldActorBeLoaded(co
 
 	for (auto& ActorDescFilter : ActorDescFilters)
 	{
-		if (!ActorDescFilter(World, Actor))
+		if (!ActorDescFilter.Get().PassFilter(World, Actor))
 		{
+			Actor->SetUnloadedReason(ActorDescFilter.Get().GetFilterReason());
 			return false;
 		}
 	}
diff --git a/Engine/Source/Runtime/Engine/Public/WorldPartition/LoaderAdapter/LoaderAdapterPinnedActors.h b/Engine/Source/Runtime/Engine/Public/WorldPartition/LoaderAdapter/LoaderAdapterPinnedActors.h
--- a/Engine/Source/Runtime/Engine/Public/WorldPartition/LoaderAdapter/LoaderAdapterPinnedActors.h
+++ b/Engine/Source/Runtime/Engine/Public/WorldPartition/LoaderAdapter/LoaderAdapterPinnedActors.h
@@ -15,7 +15,6 @@ public:
 		: FLoaderAdapterActorList(InWorld)
 	{}
 
-	static FText GetUnloadedReason(FWorldPartitionActorDesc* InActorDesc);
 	static bool SupportsPinning(FWorldPartitionActorDesc* InActorDesc);
 	static bool SupportsPinning(AActor* InActor);
 
diff --git a/Engine/Source/Runtime/Engine/Public/WorldPartition/WorldPartitionActorDesc.h b/Engine/Source/Runtime/Engine/Public/WorldPartition/WorldPartitionActorDesc.h
--- a/Engine/Source/Runtime/Engine/Public/WorldPartition/WorldPartitionActorDesc.h
+++ b/Engine/Source/Runtime/Engine/Public/WorldPartition/WorldPartitionActorDesc.h
@@ -176,9 +176,11 @@ PRAGMA_ENABLE_DEPRECATION_WARNINGS
 		return GetTypeHash(Key.Guid);
 	}
 
-	bool FailedToLoad() const
+	const FText& GetUnloadedReason() const;
+
+	void SetUnloadedReason(FText* InUnloadedReason)
 	{
-		return bFailedToLoad;
+		UnloadedReason = InUnloadedReason;
 	}
 
 protected:
@@ -306,7 +308,7 @@ protected:
 	UActorDescContainer*			Container;
 	TOptional<TArray<FName>>		ResolvedDataLayerInstanceNames; // Can only resolve in ActorDesc if Container is not used as a template
 	bool							bIsForcedNonSpatiallyLoaded;
-	mutable bool					bFailedToLoad;
+	mutable FText*					UnloadedReason;
 
 	static TMap<TSubclassOf<AActor>, FActorDescDeprecator> Deprecators;
 };
diff --git a/Engine/Source/Runtime/Engine/Public/WorldPartition/WorldPartitionActorLoaderInterface.h b/Engine/Source/Runtime/Engine/Public/WorldPartition/WorldPartitionActorLoaderInterface.h
--- a/Engine/Source/Runtime/Engine/Public/WorldPartition/WorldPartitionActorLoaderInterface.h
+++ b/Engine/Source/Runtime/Engine/Public/WorldPartition/WorldPartitionActorLoaderInterface.h
@@ -73,8 +73,18 @@ public:
 
 	virtual ILoaderAdapter* GetLoaderAdapter() =0;
 
-	using FActorDescFilter = TFunction<bool(class UWorld*, const FWorldPartitionHandle&)>;
-	static void RegisterActorDescFilter(const FActorDescFilter& InActorDescFilter);
+	class FActorDescFilter
+	{
+	public:
+		virtual ~FActorDescFilter() {}
+		virtual bool PassFilter(class UWorld*, const FWorldPartitionHandle&) = 0;
+
+		// Higher priority filters are called first
+		virtual uint32 GetFilterPriority() const = 0;
+		virtual FText* GetFilterReason() const = 0;
+	};
+		
+	static void RegisterActorDescFilter(const TSharedRef<FActorDescFilter>& InActorDescFilter);
 
 	static void RefreshLoadedState(bool bIsFromUserChange);
 
@@ -82,7 +92,7 @@ private:
 	DECLARE_EVENT_OneParam(IWorldPartitionActorLoaderInterface, FOnActorLoaderInterfaceRefreshState, bool /*bIsFromUserChange*/);
 	static FOnActorLoaderInterfaceRefreshState ActorLoaderInterfaceRefreshState;
 
-	static TArray<FActorDescFilter> ActorDescFilters;
+	static TArray<TSharedRef<FActorDescFilter>> ActorDescFilters;
 	friend class ILoaderAdapter;
 #endif
 };
\ No newline at end of file
