Actor Filter: Fix Apply to multiple instances always applying last value to all instances

#rnx

diff --git a/Engine/Source/Editor/LevelInstanceEditor/Private/LevelInstanceFilterPropertyTypeCustomization.cpp b/Engine/Source/Editor/LevelInstanceEditor/Private/LevelInstanceFilterPropertyTypeCustomization.cpp
--- a/Engine/Source/Editor/LevelInstanceEditor/Private/LevelInstanceFilterPropertyTypeCustomization.cpp
+++ b/Engine/Source/Editor/LevelInstanceEditor/Private/LevelInstanceFilterPropertyTypeCustomization.cpp
@@ -24,13 +24,14 @@ TSharedPtr<FWorldPartitionActorFilterMode::FFilter> FLevelInstanceFilterProperty
 	// Find Selected Level Instances with matching WorldAssetPackage
 	FString WorldAssetPackage;
 	UWorld* World = nullptr;
-	for (UObject* OuterObject : OuterObjects)
+	for (int32 OuterObjectIndex = 0; OuterObjectIndex < OuterObjects.Num(); ++OuterObjectIndex)
 	{
+		UObject* OuterObject = OuterObjects[OuterObjectIndex];
 		if (AActor* OuterActor = OuterObject->GetTypedOuter<AActor>())
 		{
 			if (ILevelInstanceInterface* LevelInstanceInterface = Cast<ILevelInstanceInterface>(OuterActor))
 			{
-				LevelInstances.Add(LevelInstanceInterface);
+				LevelInstances.Add(LevelInstanceInterface, OuterObjectIndex);
 				ActorLabel = OuterActor->GetActorLabel();
 
 				if (WorldAssetPackage.IsEmpty())
@@ -78,7 +79,7 @@ TSharedPtr<FWorldPartitionActorFilterMode::FFilter> FLevelInstanceFilterProperty
 
 	// Gather Filters for selected level instances
 	TArray<const FWorldPartitionActorFilter*> SelectedFilters;
-	Algo::Transform(LevelInstances, SelectedFilters, [](ILevelInstanceInterface* LevelInstanceInterface) { return &LevelInstanceInterface->GetFilter(); });
+	Algo::Transform(LevelInstances, SelectedFilters, [](const auto& Pair) { return &Pair.Key->GetFilter(); });
 
 	// Create Mode Filter which holds the final values for the filter
 	return MakeShared<FWorldPartitionActorFilterMode::FFilter>(Filter, SelectedFilters);
@@ -87,14 +88,14 @@ TSharedPtr<FWorldPartitionActorFilterMode::FFilter> FLevelInstanceFilterProperty
 void FLevelInstanceFilterPropertyTypeCustomization::ApplyFilter(TSharedRef<IPropertyHandle> PropertyHandle, const FWorldPartitionActorFilterMode& Mode)
 {
 	const FScopedTransaction Transaction(LOCTEXT("WorldPartitionActorFilterApply_Transaction", "Apply Level Instance Filter"));
-	for (ILevelInstanceInterface* LevelInstance : LevelInstances)
+	for (auto& [LevelInstance, ObjectIndex] : LevelInstances)
 	{
 		ULevelInstanceComponent* Component = LevelInstance->GetLevelInstanceComponent();
 		FWorldPartitionActorFilter ComponentFilter = Component->GetFilter();
 		Mode.Apply(ComponentFilter);
 		
 		// Setting value through PropertyHandle will take care of propagating to instances if we are editing a template (Blueprint)
-		PropertyHandle->SetValueFromFormattedString(ComponentFilter.ToString(), EPropertyValueSetFlags::DefaultFlags);
+		PropertyHandle->SetPerObjectValue(ObjectIndex, ComponentFilter.ToString(), EPropertyValueSetFlags::DefaultFlags);
 	}
 }
 
diff --git a/Engine/Source/Editor/LevelInstanceEditor/Private/LevelInstanceFilterPropertyTypeCustomization.h b/Engine/Source/Editor/LevelInstanceEditor/Private/LevelInstanceFilterPropertyTypeCustomization.h
--- a/Engine/Source/Editor/LevelInstanceEditor/Private/LevelInstanceFilterPropertyTypeCustomization.h
+++ b/Engine/Source/Editor/LevelInstanceEditor/Private/LevelInstanceFilterPropertyTypeCustomization.h
@@ -32,5 +32,5 @@ private:
 	virtual TSharedPtr<FWorldPartitionActorFilterMode::FFilter> CreateModeFilter(TArray<UObject*> OuterObjects) override;
 	virtual void ApplyFilter(TSharedRef<IPropertyHandle> PropertyHandle, const FWorldPartitionActorFilterMode& Mode) override;
 
-	TArray<ILevelInstanceInterface*> LevelInstances;
+	TMap<ILevelInstanceInterface*, int32> LevelInstances;
 };
