LevelInstance: Editor Mode Overlay

#rnx

diff --git a/Engine/Source/Editor/LevelInstanceEditor/Private/LevelInstanceEditorModeToolkit.cpp b/Engine/Source/Editor/LevelInstanceEditor/Private/LevelInstanceEditorModeToolkit.cpp
--- a/Engine/Source/Editor/LevelInstanceEditor/Private/LevelInstanceEditorModeToolkit.cpp
+++ b/Engine/Source/Editor/LevelInstanceEditor/Private/LevelInstanceEditorModeToolkit.cpp
@@ -1,7 +1,17 @@
 // Copyright Epic Games, Inc. All Rights Reserved.
 #include "LevelInstanceEditorModeToolkit.h"
-
+#include "LevelInstanceEditorMode.h"
 #include "Internationalization/Internationalization.h"
+#include "Widgets/SWidget.h"
+#include "Widgets/Text/STextBlock.h"
+#include "Widgets/Images/SImage.h"
+#include "Widgets/Input/SButton.h"
+#include "Tools/UEdMode.h"
+#include "Toolkits/IToolkitHost.h"
+#include "Styling/SlateIconFinder.h"
+#include "LevelInstance/LevelInstanceActor.h"
+#include "LevelInstance/ILevelInstanceEditorModule.h"
+#include "Engine/World.h"
 
 class FAssetEditorModeUILayer;
 
@@ -11,6 +21,79 @@ FLevelInstanceEditorModeToolkit::FLevelInstanceEditorModeToolkit()
 {
 }
 
+FLevelInstanceEditorModeToolkit::~FLevelInstanceEditorModeToolkit()
+{
+	if (ViewportOverlayWidget.IsValid())
+	{
+		GetToolkitHost()->RemoveViewportOverlayWidget(ViewportOverlayWidget.ToSharedRef());
+	}
+}
+
+void FLevelInstanceEditorModeToolkit::Init(const TSharedPtr<IToolkitHost>& InitToolkitHost, TWeakObjectPtr<UEdMode> InOwningMode)
+{
+	FModeToolkit::Init(InitToolkitHost, InOwningMode);
+
+	const ULevelInstanceSubsystem* LevelInstanceSubsystem = UWorld::GetSubsystem<ULevelInstanceSubsystem>(InitToolkitHost->GetWorld());
+	check(LevelInstanceSubsystem);
+
+	const ILevelInstanceInterface* EditingLevelInstance = LevelInstanceSubsystem->GetEditingLevelInstance();
+	check(EditingLevelInstance);
+	const TSoftObjectPtr<UWorld> WorldAsset = EditingLevelInstance->GetWorldAsset();
+	const FString WorldAssetName = WorldAsset.GetAssetName();
+	FText DisplayText = FText::FromString(WorldAsset.GetAssetName());
+
+	ILevelInstanceEditorModule& EditorModule = FModuleManager::GetModuleChecked<ILevelInstanceEditorModule>("LevelInstanceEditor");
+	
+	// ViewportOverlay
+	SAssignNew(ViewportOverlayWidget, SHorizontalBox)
+		+SHorizontalBox::Slot()
+		.HAlign(HAlign_Center)
+		.VAlign(VAlign_Bottom)
+		.Padding(FMargin(0.0f, 0.0f, 0.f, 15.f))
+		[
+			SNew(SBorder)
+			.BorderImage(FAppStyle::Get().GetBrush("EditorViewport.OverlayBrush"))
+			.Padding(8.f)
+			[
+				SNew(SHorizontalBox)
+				+SHorizontalBox::Slot()
+				.AutoWidth()
+				.VAlign(VAlign_Center)
+				.Padding(FMargin(0.f, 0.f, 8.f, 0.f))
+				[
+					SNew(SImage)
+					.Image(FSlateIconFinder::FindIconBrushForClass(ALevelInstance::StaticClass()))
+				]
+				+SHorizontalBox::Slot()
+				.AutoWidth()
+				.VAlign(VAlign_Center)
+				.Padding(FMargin(0.f, 0.f, 8.f, 0.f))
+				[
+					SNew(STextBlock)
+					.Text(DisplayText)
+				]
+				+SHorizontalBox::Slot()
+				.AutoWidth()
+				.Padding(FMargin(2.0, 0.f, 0.f, 0.f))
+				[
+					SNew(SButton)
+					.ButtonStyle(FAppStyle::Get(), "PrimaryButton")
+					.TextStyle(FAppStyle::Get(), "DialogButtonText")
+					.Text(LOCTEXT("ExitEdit", "Exit"))
+					.ToolTipText(LOCTEXT("ExitTooltip", "Exit Level Instance Edit"))
+					.HAlign(HAlign_Center)
+					.OnClicked_Lambda([&EditorModule]() 
+					{ 
+						EditorModule.BroadcastTryExitEditorMode(); 
+						return FReply::Handled(); 
+					})
+				]
+			]
+		];
+
+	GetToolkitHost()->AddViewportOverlayWidget(ViewportOverlayWidget.ToSharedRef());
+}
+
 FName FLevelInstanceEditorModeToolkit::GetToolkitFName() const
 {
 	return FName("LevelInstanceEditorModeToolkit");
@@ -21,9 +104,9 @@ FText FLevelInstanceEditorModeToolkit::GetBaseToolkitName() const
 	return LOCTEXT("ToolkitDisplayName", "Level Instance Editor Mode");
 }
 
-void FLevelInstanceEditorModeToolkit::SetModeUILayer(const TSharedPtr<FAssetEditorModeUILayer> InLayer)
+void FLevelInstanceEditorModeToolkit::RequestModeUITabs()
 {
-	// Do nothing (prevents Ed Mode Tab spawn)
+	// No Tabs
 }
 
 #undef LOCTEXT_NAMESPACE
\ No newline at end of file
diff --git a/Engine/Source/Editor/LevelInstanceEditor/Private/LevelInstanceEditorModeToolkit.h b/Engine/Source/Editor/LevelInstanceEditor/Private/LevelInstanceEditorModeToolkit.h
--- a/Engine/Source/Editor/LevelInstanceEditor/Private/LevelInstanceEditorModeToolkit.h
+++ b/Engine/Source/Editor/LevelInstanceEditor/Private/LevelInstanceEditorModeToolkit.h
@@ -6,16 +6,24 @@
 #include "Toolkits/BaseToolkit.h"
 #include "UObject/NameTypes.h"
 
-class FAssetEditorModeUILayer;
+class IToolkitHost;
+class SWidget;
+class UEdMode;
 
 class FLevelInstanceEditorModeToolkit : public FModeToolkit
 {
 public:
 	FLevelInstanceEditorModeToolkit();
+	virtual ~FLevelInstanceEditorModeToolkit();
 
 	// IToolkit interface 
+	virtual void Init(const TSharedPtr<IToolkitHost>& InitToolkitHost, TWeakObjectPtr<UEdMode> InOwningMode);
+
 	virtual FName GetToolkitFName() const override;
 	virtual FText GetBaseToolkitName() const override;
 
-	virtual void SetModeUILayer(const TSharedPtr<FAssetEditorModeUILayer> InLayer) override;
+	virtual void RequestModeUITabs() override;
+
+private:
+	TSharedPtr<SWidget> ViewportOverlayWidget;
 };
\ No newline at end of file
