LevelInstanceEditorModeToolkit: Fix crash on editor exit (while editing LevelInstance)

#rnx

diff --git a/Engine/Source/Editor/LevelInstanceEditor/Private/LevelInstanceEditorModeToolkit.cpp b/Engine/Source/Editor/LevelInstanceEditor/Private/LevelInstanceEditorModeToolkit.cpp
--- a/Engine/Source/Editor/LevelInstanceEditor/Private/LevelInstanceEditorModeToolkit.cpp
+++ b/Engine/Source/Editor/LevelInstanceEditor/Private/LevelInstanceEditorModeToolkit.cpp
@@ -23,7 +23,7 @@ FLevelInstanceEditorModeToolkit::FLevelInstanceEditorModeToolkit()
 
 FLevelInstanceEditorModeToolkit::~FLevelInstanceEditorModeToolkit()
 {
-	if (ViewportOverlayWidget.IsValid())
+	if(IsHosted() && ViewportOverlayWidget.IsValid())
 	{
 		GetToolkitHost()->RemoveViewportOverlayWidget(ViewportOverlayWidget.ToSharedRef());
 	}
