LevelInstance: Now that we have a Mode overlay UI (Shift+L will remove postprocess)

#rnx

diff --git a/Engine/Source/Editor/LevelInstanceEditor/Private/LevelInstanceEditorMode.cpp b/Engine/Source/Editor/LevelInstanceEditor/Private/LevelInstanceEditorMode.cpp
--- a/Engine/Source/Editor/LevelInstanceEditor/Private/LevelInstanceEditorMode.cpp
+++ b/Engine/Source/Editor/LevelInstanceEditor/Private/LevelInstanceEditorMode.cpp
@@ -48,7 +48,7 @@ void ULevelInstanceEditorMode::UpdateEngineShowFlags()
 		{
 			if(ULevelInstanceSubsystem* LevelInstanceSubsystem = LevelVC->GetWorld()->GetSubsystem<ULevelInstanceSubsystem>())
 			{
-				const bool bEditingLevelInstance = !!LevelInstanceSubsystem->GetEditingLevelInstance();
+				const bool bEditingLevelInstance = !!LevelInstanceSubsystem->GetEditingLevelInstance() && IsContextRestrictionCommandEnabled();
 				// Make sure we update both Game/Editor showflags
 				LevelVC->EngineShowFlags.EditingLevelInstance = bEditingLevelInstance;
 				LevelVC->LastEngineShowFlags.EditingLevelInstance = bEditingLevelInstance;
@@ -181,6 +181,8 @@ void ULevelInstanceEditorMode::ExitModeCommand()
 void ULevelInstanceEditorMode::ToggleContextRestrictionCommand()
 {
 	bContextRestriction = !bContextRestriction;
+
+	UpdateEngineShowFlags();
 }
 
 bool ULevelInstanceEditorMode::IsContextRestrictionCommandEnabled() const
