Packed Level Actors: New Update Packed Blueprints CB action on World assets
- Move all static Packing related methods to new Utils class
- Remove dependency on ALevelInstance (use ILevelInstanceInterface)

#rnx

diff --git a/Engine/Source/Editor/LevelInstanceEditor/Private/LevelInstanceEditorModule.cpp b/Engine/Source/Editor/LevelInstanceEditor/Private/LevelInstanceEditorModule.cpp
--- a/Engine/Source/Editor/LevelInstanceEditor/Private/LevelInstanceEditorModule.cpp
+++ b/Engine/Source/Editor/LevelInstanceEditor/Private/LevelInstanceEditorModule.cpp
@@ -2,6 +2,7 @@
 #include "LevelInstanceEditorModule.h"
 #include "LevelInstanceActorDetails.h"
 #include "LevelInstancePivotDetails.h"
+#include "PackedLevelActorUtils.h"
 #include "LevelInstanceFilterPropertyTypeCustomization.h"
 #include "LevelInstance/LevelInstanceSubsystem.h"
 #include "LevelInstance/LevelInstanceInterface.h"
@@ -495,13 +496,13 @@ namespace LevelInstanceMenuUtils
 					{
 						FToolUIAction UIAction;
 						UIAction.ExecuteAction.BindLambda([ContextLevelInstance, BlueprintAsset](const FToolMenuContext& MenuContext)
-							{
-								APackedLevelActor::CreateOrUpdateBlueprint(ContextLevelInstance->GetWorldAsset(), BlueprintAsset);
-							});
+						{
+							FPackedLevelActorUtils::CreateOrUpdateBlueprint(ContextLevelInstance->GetWorldAsset(), BlueprintAsset);
+						});
 						UIAction.CanExecuteAction.BindLambda([](const FToolMenuContext& MenuContext)
-							{
-								return GEditor->GetSelectedActorCount() > 0;
-							});
+						{
+							return FPackedLevelActorUtils::CanPack() && GEditor->GetSelectedActorCount() > 0;
+						});
 
 						Section.AddMenuEntry(
 							"UpdatePackedBlueprint",
@@ -567,7 +568,7 @@ namespace LevelInstanceMenuUtils
 			
 			if (NewBlueprint->GeneratedClass->IsChildOf<APackedLevelActor>())
 			{
-				APackedLevelActor::UpdateBlueprint(NewBlueprint);
+				FPackedLevelActorUtils::UpdateBlueprint(NewBlueprint);
 			}
 
 			FContentBrowserModule& ContentBrowserModule = FModuleManager::LoadModuleChecked<FContentBrowserModule>("ContentBrowser");
@@ -593,7 +594,7 @@ namespace LevelInstanceMenuUtils
 			"CreateLevelInstanceBlueprint",
 			LOCTEXT("CreateLevelInstanceBlueprint", "New Blueprint..."),
 			TAttribute<FText>(),
-			TAttribute<FSlateIcon>(),
+			FSlateIcon(FAppStyle::GetAppStyleSetName(), "ClassIcon.LevelInstance"),
 			UIAction);
 	}
 
@@ -655,6 +656,37 @@ namespace LevelInstanceMenuUtils
 		}
 	}
 
+	void UpdatePackedBlueprintsFromMenu(UToolMenu* Menu, FAssetData WorldAsset)
+	{
+		FToolMenuSection& Section = CreateLevelSection(Menu);
+		FToolUIAction UIAction;
+		UIAction.CanExecuteAction.BindLambda([](const FToolMenuContext& MenuContext)
+		{
+			return FPackedLevelActorUtils::CanPack();
+		});
+		UIAction.ExecuteAction.BindLambda([WorldAsset](const FToolMenuContext& MenuContext)
+		{
+			FScopedSlowTask SlowTask(0.0f, LOCTEXT("UpdatePackedBlueprintsProgress", "Updating Packed Blueprints..."));
+			TSet<TSoftObjectPtr<UBlueprint>> BlueprintAssets;
+			FPackedLevelActorUtils::GetPackedBlueprintsForWorldAsset(TSoftObjectPtr<UWorld>(WorldAsset.GetSoftObjectPath()), BlueprintAssets, false);
+			for (TSoftObjectPtr<UBlueprint> BlueprintAsset : BlueprintAssets)
+			{
+				if (UBlueprint* Blueprint = BlueprintAsset.Get())
+				{
+					FPackedLevelActorUtils::UpdateBlueprint(Blueprint, false);
+				}
+			}
+		});
+
+		Section.AddMenuEntry(
+			"UpdatePackedBlueprintsFromMenu",
+			LOCTEXT("UpdatePackedBlueprintsFromMenu", "Update Packed Blueprints"),
+			TAttribute<FText>(),
+			FSlateIcon(FAppStyle::GetAppStyleSetName(), "ClassIcon.PackedLevelActor"),
+			UIAction
+		);
+	}
+
 	void AddPartitionedStreamingSupportFromMenu(UToolMenu* Menu, FAssetData WorldAsset)
 	{
 		FName WorldAssetName = WorldAsset.PackageName;
@@ -815,20 +847,11 @@ void FLevelInstanceEditorModule::ExtendContextMenu()
 		FUIAction PackAction(
 			FExecuteAction::CreateLambda([]() 
 			{
-				UWorld* World = GEditor->GetEditorWorldContext().World();
-				if (ULevelInstanceSubsystem* LevelInstanceSubsystem = World->GetSubsystem<ULevelInstanceSubsystem>())
-				{
-					LevelInstanceSubsystem->PackAllLoadedActors();
-				}
+				FPackedLevelActorUtils::PackAllLoadedActors();
 			}), 
 			FCanExecuteAction::CreateLambda([]()
 			{
-				UWorld* World = GEditor->GetEditorWorldContext().World();
-				if (ULevelInstanceSubsystem* LevelInstanceSubsystem = World->GetSubsystem<ULevelInstanceSubsystem>())
-				{
-					return LevelInstanceSubsystem->CanPackAllLoadedActors();
-				}
-				return false;
+				return FPackedLevelActorUtils::CanPack();
 			}),
 			FIsActionChecked(),
 			FIsActionButtonVisible::CreateLambda([]() 
@@ -842,6 +865,11 @@ void FLevelInstanceEditorModule::ExtendContextMenu()
 
 	auto AddDynamicSection = [](UToolMenu* ToolMenu)
 	{				
+		if (GEditor->GetPIEWorldContext())
+		{
+			return;
+		}
+
 		if (ULevelEditorContextMenuContext* LevelEditorMenuContext = ToolMenu->Context.FindContext<ULevelEditorContextMenuContext>())
 		{
 			// Use the actor under the cursor if available (e.g. right-click menu).
@@ -881,6 +909,11 @@ void FLevelInstanceEditorModule::ExtendContextMenu()
 	{
 		FToolMenuSection& Section = WorldAssetMenu->AddDynamicSection("ActorLevelInstance", FNewToolMenuDelegate::CreateLambda([this](UToolMenu* ToolMenu)
 		{
+			if (GEditor->GetPIEWorldContext())
+			{
+				return;
+			}
+
 			if(!GetDefault<UEditorExperimentalSettings>()->bLevelInstance)
 			{
 				return;
@@ -899,6 +932,7 @@ void FLevelInstanceEditorModule::ExtendContextMenu()
 					if (AssetMenuContext->SelectedAssets[0].IsInstanceOf<UWorld>())
 					{
 						LevelInstanceMenuUtils::CreateBlueprintFromMenu(ToolMenu, WorldAsset);
+						LevelInstanceMenuUtils::UpdatePackedBlueprintsFromMenu(ToolMenu, WorldAsset);
 						LevelInstanceMenuUtils::AddPartitionedStreamingSupportFromMenu(ToolMenu, WorldAsset);
 					}
 				}
diff --git a/Engine/Source/Editor/LevelInstanceEditor/Private/PackedLevelActorUtils.cpp b/Engine/Source/Editor/LevelInstanceEditor/Private/PackedLevelActorUtils.cpp
new file mode 100644
--- /dev/null
+++ b/Engine/Source/Editor/LevelInstanceEditor/Private/PackedLevelActorUtils.cpp
@@ -0,0 +1,186 @@
+// Copyright Epic Games, Inc. All Rights Reserved.
+#if WITH_EDITOR
+
+#include "PackedLevelActorUtils.h"
+#include "PackedLevelActor/PackedLevelActor.h"
+#include "PackedLevelActor/PackedLevelActorBuilder.h"
+#include "LevelInstance/LevelInstanceInterface.h"
+#include "AssetRegistry/IAssetRegistry.h"
+#include "AssetRegistry/AssetRegistryModule.h"
+#include "Engine/World.h"
+#include "Engine/Blueprint.h"
+#include "EngineUtils.h"
+#include "Misc/ScopedSlowTask.h"
+#include "Editor.h"
+#include "UObject/UObjectIterator.h"
+
+#define LOCTEXT_NAMESPACE "FPackedLevelActorUtils"
+
+void FPackedLevelActorUtils::GetPackedBlueprintsForWorldAsset(const TSoftObjectPtr<UWorld>& InWorldAsset, TSet<TSoftObjectPtr<UBlueprint>>& OutPackedBlueprintAssets, bool bInLoadedOnly, bool bInDiskOnly)
+{
+	FAssetRegistryModule& AssetRegistryModule = FModuleManager::LoadModuleChecked<FAssetRegistryModule>(FName("AssetRegistry"));
+	IAssetRegistry& AssetRegistry = AssetRegistryModule.Get();
+
+	if (InWorldAsset.IsNull())
+	{
+		return;
+	}
+			
+	TArray<FSoftObjectPath> AssetsToProcess;
+	AssetsToProcess.Add(InWorldAsset.ToSoftObjectPath());
+	TSet<FSoftObjectPath> ProcessedAssets;
+
+	while (!AssetsToProcess.IsEmpty())
+	{
+		FSoftObjectPath ProcessingAsset = AssetsToProcess.Pop();
+		ProcessedAssets.Add(ProcessingAsset);
+
+		TArray<FName> ReferencerPackages;
+		AssetRegistry.GetReferencers(*ProcessingAsset.GetLongPackageName(), ReferencerPackages);
+
+		for (FName ReferencerPackage : ReferencerPackages)
+		{
+			TArray<FAssetData> Assets;
+			AssetRegistry.GetAssetsByPackageName(ReferencerPackage, Assets);
+
+			for (const FAssetData& Asset : Assets)
+			{
+				if (UClass* AssetClass = Asset.GetClass(bInLoadedOnly ? EResolveClass::No : EResolveClass::Yes))
+				{
+					if (TSubclassOf<UBlueprint> BlueprintClass = AssetClass)
+					{
+						if (UBlueprint* Blueprint = Cast<UBlueprint>(Asset.GetAsset()); Blueprint && Blueprint->GeneratedClass && Blueprint->GeneratedClass->IsChildOf<APackedLevelActor>())
+						{
+							FSoftObjectPath BlueprintAsset(Blueprint);
+							if (!ProcessedAssets.Contains(BlueprintAsset) && ProcessedAssets.Contains(Blueprint->GeneratedClass->GetDefaultObject<APackedLevelActor>()->GetWorldAsset().ToSoftObjectPath()))
+							{
+								OutPackedBlueprintAssets.Add(Blueprint);
+								AssetsToProcess.Add(BlueprintAsset);
+							}						
+						}
+					}
+					else if (TSubclassOf<AActor> ActorClass = AssetClass; ActorClass && ActorClass->ImplementsInterface(ULevelInstanceInterface::StaticClass()))
+					{
+						FSoftObjectPath ReferencingAsset(Asset.GetSoftObjectPath().GetWithoutSubPath());
+						if (!ProcessedAssets.Contains(ReferencingAsset))
+						{
+							AssetsToProcess.Add(ReferencingAsset);
+						}
+					}
+					else if (TSubclassOf<UWorld> WorldClass = AssetClass)
+					{
+						FSoftObjectPath ReferencingAsset(Asset.GetSoftObjectPath());
+						if (!ProcessedAssets.Contains(ReferencingAsset))
+						{
+							AssetsToProcess.Add(ReferencingAsset);
+						}
+					}
+				}
+			}
+		}
+	}
+}
+
+bool FPackedLevelActorUtils::CanPack()
+{
+	if (!GEditor)
+	{
+		return false;
+	}
+
+	if (GEditor->GetPIEWorldContext())
+	{
+		return false;
+	}
+
+	if (ULevelInstanceSubsystem* LevelInstanceSubsystem = GEditor->GetEditorWorldContext().World()->GetSubsystem<ULevelInstanceSubsystem>())
+	{
+		return !LevelInstanceSubsystem->GetEditingLevelInstance();
+	}
+
+	return false;
+}
+
+void FPackedLevelActorUtils::PackAllLoadedActors()
+{
+	if (!CanPack())
+	{
+		return;
+	}
+	
+	TSet<APackedLevelActor*> PackedLevelActorsToUpdate;
+	TSet<UBlueprint*> BlueprintsToUpdate;
+
+	if (UWorld* EditorWorld = GEditor->GetEditorWorldContext().World())
+	{
+		for (TActorIterator<APackedLevelActor> PackedLevelActorIt(EditorWorld); PackedLevelActorIt; ++PackedLevelActorIt)
+		{
+			APackedLevelActor* PackedLevelActor = *PackedLevelActorIt;
+			UBlueprint* Blueprint = PackedLevelActor->GetRootBlueprint();
+			if (Blueprint)
+			{
+				BlueprintsToUpdate.Add(Blueprint);
+			}
+			else
+			{
+				PackedLevelActorsToUpdate.Add(PackedLevelActor);
+			}
+		}
+	}
+	
+
+	int32 Count = BlueprintsToUpdate.Num() + PackedLevelActorsToUpdate.Num();
+	if (!Count)
+	{
+		return;
+	}
+
+	if (GEditor)
+	{
+		GEditor->SelectNone(true, true);
+	}
+
+	FScopedSlowTask SlowTask(Count, (LOCTEXT("TaskPackLevels", "Packing Levels")));
+	SlowTask.MakeDialog();
+
+	auto UpdateProgress = [&SlowTask]()
+	{
+		if (SlowTask.CompletedWork < SlowTask.TotalAmountOfWork)
+		{
+			SlowTask.EnterProgressFrame(1, FText::Format(LOCTEXT("TaskPackLevelProgress", "Packing Level {0} of {1}"), FText::AsNumber(SlowTask.CompletedWork), FText::AsNumber(SlowTask.TotalAmountOfWork)));
+		}
+	};
+
+	TSharedPtr<FPackedLevelActorBuilder> Builder = FPackedLevelActorBuilder::CreateDefaultBuilder();
+	const bool bCheckoutAndSave = false;
+	for (UBlueprint* Blueprint : BlueprintsToUpdate)
+	{
+		Builder->UpdateBlueprint(Blueprint, bCheckoutAndSave);
+		UpdateProgress();
+	}
+
+	for (APackedLevelActor* PackedLevelActor : PackedLevelActorsToUpdate)
+	{
+		PackedLevelActor->UpdateLevelInstanceFromWorldAsset();
+		UpdateProgress();
+	}
+}
+
+bool FPackedLevelActorUtils::CreateOrUpdateBlueprint(ILevelInstanceInterface* InLevelInstance, TSoftObjectPtr<UBlueprint> InBlueprintAsset, bool bCheckoutAndSave, bool bPromptForSave)
+{
+	return FPackedLevelActorBuilder::CreateDefaultBuilder()->CreateOrUpdateBlueprint(InLevelInstance, InBlueprintAsset, bCheckoutAndSave, bPromptForSave);
+}
+
+bool FPackedLevelActorUtils::CreateOrUpdateBlueprint(TSoftObjectPtr<UWorld> InWorldAsset, TSoftObjectPtr<UBlueprint> InBlueprintAsset, bool bCheckoutAndSave, bool bPromptForSave)
+{
+	return FPackedLevelActorBuilder::CreateDefaultBuilder()->CreateOrUpdateBlueprint(InWorldAsset, InBlueprintAsset, bCheckoutAndSave, bPromptForSave);
+}
+
+void FPackedLevelActorUtils::UpdateBlueprint(UBlueprint* InBlueprint, bool bCheckoutAndSave)
+{
+	return FPackedLevelActorBuilder::CreateDefaultBuilder()->UpdateBlueprint(InBlueprint, bCheckoutAndSave);
+}
+
+#undef LOCTEXT_NAMESPACE
+
+#endif
\ No newline at end of file
diff --git a/Engine/Source/Editor/LevelInstanceEditor/Public/PackedLevelActorUtils.h b/Engine/Source/Editor/LevelInstanceEditor/Public/PackedLevelActorUtils.h
new file mode 100644
--- /dev/null
+++ b/Engine/Source/Editor/LevelInstanceEditor/Public/PackedLevelActorUtils.h
@@ -0,0 +1,28 @@
+// Copyright Epic Games, Inc. All Rights Reserved.
+#pragma once
+
+#if WITH_EDITOR
+
+#include "UObject/SoftObjectPtr.h"
+
+class UWorld;
+class UBlueprint;
+class ILevelInstanceInterface;
+
+class FPackedLevelActorUtils
+{
+public:
+	static void GetPackedBlueprintsForWorldAsset(const TSoftObjectPtr<UWorld>& InWorldAsset, TSet<TSoftObjectPtr<UBlueprint>>& OutPackedBlueprintAssets, bool bInLoadedOnly = false, bool bInDiskOnly = false);
+
+	static bool CanPack();
+
+	static void PackAllLoadedActors();
+
+	static bool CreateOrUpdateBlueprint(ILevelInstanceInterface* InLevelInstance, TSoftObjectPtr<UBlueprint> InBlueprintAsset, bool bCheckoutAndSave = true, bool bPromptForSave = true);
+
+	static bool CreateOrUpdateBlueprint(TSoftObjectPtr<UWorld> InWorldAsset, TSoftObjectPtr<UBlueprint> InBlueprintAsset, bool bCheckoutAndSave = true, bool bPromptForSave = true);
+	
+	static void UpdateBlueprint(UBlueprint* InBlueprint, bool bCheckoutAndSave = true);
+};
+
+#endif 
\ No newline at end of file
diff --git a/Engine/Source/Runtime/Engine/Private/LevelInstance/LevelInstanceSubsystem.cpp b/Engine/Source/Runtime/Engine/Private/LevelInstance/LevelInstanceSubsystem.cpp
--- a/Engine/Source/Runtime/Engine/Private/LevelInstance/LevelInstanceSubsystem.cpp
+++ b/Engine/Source/Runtime/Engine/Private/LevelInstance/LevelInstanceSubsystem.cpp
@@ -655,99 +655,6 @@ bool ULevelInstanceSubsystem::OnExitEditorModeInternal(bool bForceExit)
 	return false;
 }
 
-bool ULevelInstanceSubsystem::CanPackAllLoadedActors() const
-{
-	return !LevelInstanceEdit;
-}
-
-void ULevelInstanceSubsystem::PackAllLoadedActors()
-{
-	if (!CanPackAllLoadedActors())
-	{
-		return;
-	}
-
-	// Add Dependencies first so that we pack in the proper order (depth first)
-	TFunction<void(APackedLevelActor*, TArray<UBlueprint*>&, TArray<APackedLevelActor*>&)> GatherDepencenciesRecursive = [&GatherDepencenciesRecursive](APackedLevelActor* PackedLevelActor, TArray<UBlueprint*>& BPsToPack, TArray<APackedLevelActor*>& ToPack)
-	{
-		// Early out on already processed BPs or non BP Packed LIs.
-		UBlueprint* Blueprint = PackedLevelActor->GetRootBlueprint();
-		if ((Blueprint && BPsToPack.Contains(Blueprint)) || ToPack.Contains(PackedLevelActor))
-		{
-			return;
-		}
-		
-		// Recursive deps
-		for (const TSoftObjectPtr<UBlueprint>& Dependency : PackedLevelActor->PackedBPDependencies)
-		{
-			if (UBlueprint* LoadedDependency = Dependency.LoadSynchronous())
-			{
-				if (APackedLevelActor* CDO = Cast<APackedLevelActor>(LoadedDependency->GeneratedClass ? LoadedDependency->GeneratedClass->GetDefaultObject() : nullptr))
-				{
-					GatherDepencenciesRecursive(CDO, BPsToPack, ToPack);
-				}
-			}
-		}
-
-		// Add after dependencies
-		if (Blueprint)
-		{
-			BPsToPack.Add(Blueprint);
-		}
-		else
-		{
-			ToPack.Add(PackedLevelActor);
-		}
-	};
-
-	TArray<APackedLevelActor*> PackedLevelActorsToUpdate;
-	TArray<UBlueprint*> BlueprintsToUpdate;
-	for (TObjectIterator<UWorld> It(RF_ClassDefaultObject | RF_ArchetypeObject, true); It; ++It)
-	{
-		UWorld* CurrentWorld = *It;
-		if (IsValid(CurrentWorld) && CurrentWorld->GetSubsystem<ULevelInstanceSubsystem>() != nullptr)
-		{
-			for (TActorIterator<APackedLevelActor> PackedLevelActorIt(CurrentWorld); PackedLevelActorIt; ++PackedLevelActorIt)
-			{
-				GatherDepencenciesRecursive(*PackedLevelActorIt, BlueprintsToUpdate, PackedLevelActorsToUpdate);
-			}
-		}
-	}
-
-	int32 Count = BlueprintsToUpdate.Num() + PackedLevelActorsToUpdate.Num();
-	if (!Count)
-	{
-		return;
-	}
-	
-	GEditor->SelectNone(true, true);
-
-	FScopedSlowTask SlowTask(Count, (LOCTEXT("TaskPackLevels", "Packing Levels")));
-	SlowTask.MakeDialog();
-		
-	auto UpdateProgress = [&SlowTask]()
-	{
-		if (SlowTask.CompletedWork < SlowTask.TotalAmountOfWork)
-		{
-			SlowTask.EnterProgressFrame(1, FText::Format(LOCTEXT("TaskPackLevelProgress", "Packing Level {0} of {1}"), FText::AsNumber(SlowTask.CompletedWork), FText::AsNumber(SlowTask.TotalAmountOfWork)));
-		}
-	};
-
-	TSharedPtr<FPackedLevelActorBuilder> Builder = FPackedLevelActorBuilder::CreateDefaultBuilder();
-	const bool bCheckoutAndSave = false;
-	for (UBlueprint* Blueprint : BlueprintsToUpdate)
-	{
-		Builder->UpdateBlueprint(Blueprint, bCheckoutAndSave);
-		UpdateProgress();
-	}
-
-	for (APackedLevelActor* PackedLevelActor : PackedLevelActorsToUpdate)
-	{
-		PackedLevelActor->UpdateLevelInstanceFromWorldAsset();
-		UpdateProgress();
-	}
-}
-
 bool ULevelInstanceSubsystem::GetLevelInstanceBounds(const ILevelInstanceInterface* LevelInstance, FBox& OutBounds) const
 {
 	if (IsLoaded(LevelInstance))
diff --git a/Engine/Source/Runtime/Engine/Private/PackedLevelActor/IPackedLevelActorBuilder.h b/Engine/Source/Runtime/Engine/Private/PackedLevelActor/IPackedLevelActorBuilder.h
deleted file mode 100644
diff --git a/Engine/Source/Runtime/Engine/Private/PackedLevelActor/PackedLevelActor.cpp b/Engine/Source/Runtime/Engine/Private/PackedLevelActor/PackedLevelActor.cpp
--- a/Engine/Source/Runtime/Engine/Private/PackedLevelActor/PackedLevelActor.cpp
+++ b/Engine/Source/Runtime/Engine/Private/PackedLevelActor/PackedLevelActor.cpp
@@ -174,21 +174,6 @@ TUniquePtr<class FWorldPartitionActorDesc> APackedLevelActor::CreateClassActorDe
 	return TUniquePtr<FWorldPartitionActorDesc>(new FPackedLevelActorDesc());
 }
 
-bool APackedLevelActor::CreateOrUpdateBlueprint(ALevelInstance* InLevelInstance, TSoftObjectPtr<UBlueprint> InBlueprintAsset, bool bCheckoutAndSave, bool bPromptForSave)
-{
-	return FPackedLevelActorBuilder::CreateDefaultBuilder()->CreateOrUpdateBlueprint(InLevelInstance, InBlueprintAsset, bCheckoutAndSave, bPromptForSave);
-}
-
-bool APackedLevelActor::CreateOrUpdateBlueprint(TSoftObjectPtr<UWorld> InWorldAsset, TSoftObjectPtr<UBlueprint> InBlueprintAsset, bool bCheckoutAndSave, bool bPromptForSave)
-{
-	return FPackedLevelActorBuilder::CreateDefaultBuilder()->CreateOrUpdateBlueprint(InWorldAsset, InBlueprintAsset, bCheckoutAndSave, bPromptForSave);
-}
-
-void APackedLevelActor::UpdateBlueprint(UBlueprint* InBlueprint, bool bCheckoutAndSave)
-{
-	FPackedLevelActorBuilder::CreateDefaultBuilder()->UpdateBlueprint(InBlueprint, bCheckoutAndSave);
-}
-
 FName APackedLevelActor::GetPackedComponentTag()
 {
 	static FName PackedComponentTag("PackedComponent");
diff --git a/Engine/Source/Runtime/Engine/Private/PackedLevelActor/PackedLevelActorBuilder.cpp b/Engine/Source/Runtime/Engine/Private/PackedLevelActor/PackedLevelActorBuilder.cpp
--- a/Engine/Source/Runtime/Engine/Private/PackedLevelActor/PackedLevelActorBuilder.cpp
+++ b/Engine/Source/Runtime/Engine/Private/PackedLevelActor/PackedLevelActorBuilder.cpp
@@ -11,6 +11,7 @@
 #include "PackedLevelActor/PackedLevelActorRecursiveBuilder.h"
 
 #include "LevelInstance/LevelInstanceSubsystem.h"
+#include "LevelInstance/LevelInstanceInterface.h"
 
 #include "Kismet2/KismetEditorUtilities.h"
 #include "Engine/SCS_Node.h"
@@ -160,7 +161,7 @@ bool FPackedLevelActorBuilder::PackActor(APackedLevelActor* InPackedLevelActor)
 	return PackActor(InPackedLevelActor, InPackedLevelActor);
 }
 
-bool FPackedLevelActorBuilder::PackActor(APackedLevelActor* InPackedLevelActor, ALevelInstance* InLevelInstanceToPack)
+bool FPackedLevelActorBuilder::PackActor(APackedLevelActor* InPackedLevelActor, ILevelInstanceInterface* InLevelInstanceToPack)
 {
 	FPackedLevelActorBuilderContext Context(*this, InPackedLevelActor, InLevelInstanceToPack);
 	return PackActor(Context);
@@ -192,7 +193,7 @@ bool FPackedLevelActorBuilder::PackActor(FPackedLevelActorBuilderContext& InCont
 	FTransform RelativePivotTransform(NewPivotTransform.GetRelativeTransform(CurrentPivotTransform));
 		
 	InContext.SetRelativePivotTransform(RelativePivotTransform);	
-	InContext.ClusterLevelActor(InContext.GetLevelInstanceToPack());
+	InContext.ClusterLevelActor(CastChecked<AActor>(InContext.GetLevelInstanceToPack()));
 
 	for (const auto& Pair : InContext.GetClusters())
 	{
@@ -348,7 +349,7 @@ bool FPackedLevelActorBuilder::CreateOrUpdateBlueprint(TSoftObjectPtr<UWorld> In
 	return bResult;
 }
 
-bool FPackedLevelActorBuilder::CreateOrUpdateBlueprint(ALevelInstance* InLevelInstance, TSoftObjectPtr<UBlueprint> InBlueprintAsset, bool bCheckoutAndSave, bool bPromptForSave)
+bool FPackedLevelActorBuilder::CreateOrUpdateBlueprint(ILevelInstanceInterface* InLevelInstance, TSoftObjectPtr<UBlueprint> InBlueprintAsset, bool bCheckoutAndSave, bool bPromptForSave)
 {
 	if (APackedLevelActor* PackedLevelActor = Cast<APackedLevelActor>(InLevelInstance))
 	{
@@ -358,10 +359,11 @@ bool FPackedLevelActorBuilder::CreateOrUpdateBlueprint(ALevelInstance* InLevelIn
 	return CreateOrUpdateBlueprintFromUnpacked(InLevelInstance, InBlueprintAsset, bCheckoutAndSave, bPromptForSave);
 }
 
-bool FPackedLevelActorBuilder::CreateOrUpdateBlueprintFromUnpacked(ALevelInstance* InActor, TSoftObjectPtr<UBlueprint> InBlueprintAsset, bool bCheckoutAndSave, bool bPromptForSave)
+bool FPackedLevelActorBuilder::CreateOrUpdateBlueprintFromUnpacked(ILevelInstanceInterface* InLevelInstance, TSoftObjectPtr<UBlueprint> InBlueprintAsset, bool bCheckoutAndSave, bool bPromptForSave)
 {
 	bool bResult = true;
-	
+	AActor* LevelInstanceActor = CastChecked<AActor>(InLevelInstance);
+
 	// Create Temp Actor for Packing
 	FActorSpawnParameters SpawnParams;
 	SpawnParams.bCreateActorPackage = false;
@@ -370,18 +372,18 @@ bool FPackedLevelActorBuilder::CreateOrUpdateBlueprintFromUnpacked(ALevelInstanc
 	SpawnParams.ObjectFlags |= RF_Transient;
 	SpawnParams.ObjectFlags &= ~RF_Transactional;
 	
-	UWorld* World = InActor->GetWorld();
+	UWorld* World = LevelInstanceActor->GetWorld();
 	SpawnParams.OverrideLevel = World->PersistentLevel;
 
-	APackedLevelActor* PackedLevelActor = World->SpawnActor<APackedLevelActor>(InActor->GetActorLocation(), InActor->GetActorRotation(), SpawnParams);	
-	PackedLevelActor->SetWorldAsset(InActor->GetWorldAsset());
-	PackedLevelActor->SetFilter(InActor->GetFilter());
+	APackedLevelActor* PackedLevelActor = World->SpawnActor<APackedLevelActor>(LevelInstanceActor->GetActorLocation(), LevelInstanceActor->GetActorRotation(), SpawnParams);
+	PackedLevelActor->SetWorldAsset(InLevelInstance->GetWorldAsset());
+	PackedLevelActor->SetFilter(InLevelInstance->GetFilter());
 	ON_SCOPE_EXIT
 	{
-		InActor->GetWorld()->DestroyActor(PackedLevelActor);
+		LevelInstanceActor->GetWorld()->DestroyActor(PackedLevelActor);
 	};
 
-	if (!PackActor(PackedLevelActor, InActor))
+	if (!PackActor(PackedLevelActor, InLevelInstance))
 	{
 		return false;
 	}
diff --git a/Engine/Source/Runtime/Engine/Private/PackedLevelActor/PackedLevelActorBuilder.h b/Engine/Source/Runtime/Engine/Private/PackedLevelActor/PackedLevelActorBuilder.h
deleted file mode 100644
diff --git a/Engine/Source/Runtime/Engine/Private/PackedLevelActor/PackedLevelActorTypes.h b/Engine/Source/Runtime/Engine/Private/PackedLevelActor/PackedLevelActorTypes.h
deleted file mode 100644
diff --git a/Engine/Source/Runtime/Engine/Public/LevelInstance/LevelInstanceSubsystem.h b/Engine/Source/Runtime/Engine/Public/LevelInstance/LevelInstanceSubsystem.h
--- a/Engine/Source/Runtime/Engine/Public/LevelInstance/LevelInstanceSubsystem.h
+++ b/Engine/Source/Runtime/Engine/Public/LevelInstance/LevelInstanceSubsystem.h
@@ -67,8 +67,11 @@ public:
 	ENGINE_API void OnExitEditorMode();
 	ENGINE_API void OnTryExitEditorMode();
 	ENGINE_API bool OnExitEditorModeInternal(bool bForceExit);
-	ENGINE_API void PackAllLoadedActors();
-	ENGINE_API bool CanPackAllLoadedActors() const;
+
+	UE_DEPRECATED(5.3, "Use FPackedLevelActorUtils::PackAllLoadedActors")
+	ENGINE_API void PackAllLoadedActors() {}
+	UE_DEPRECATED(5.3, "Use FPackedLevelActorUtils::CanPackAllLoadedActors")
+	ENGINE_API bool CanPackAllLoadedActors() const { return false;}
 
 	ENGINE_API ILevelInstanceInterface* GetEditingLevelInstance() const;
 	ENGINE_API bool CanEditLevelInstance(const ILevelInstanceInterface* LevelInstance, FText* OutReason = nullptr) const;
diff --git a/Engine/Source/Runtime/Engine/Public/PackedLevelActor/IPackedLevelActorBuilder.h b/Engine/Source/Runtime/Engine/Public/PackedLevelActor/IPackedLevelActorBuilder.h
new file mode 100644
--- /dev/null
+++ b/Engine/Source/Runtime/Engine/Public/PackedLevelActor/IPackedLevelActorBuilder.h
@@ -0,0 +1,26 @@
+// Copyright Epic Games, Inc. All Rights Reserved.
+
+#pragma once
+
+#if WITH_EDITOR
+
+#include "CoreMinimal.h"
+#include "UObject/ObjectMacros.h"
+#include "PackedLevelActor/PackedLevelActorTypes.h"
+
+class APackedLevelActor;
+class AActor;
+class UActorComponent;
+class FPackedLevelActorBuilderContext;
+
+class IPackedLevelActorBuilder
+{
+public:
+	IPackedLevelActorBuilder() {}
+	virtual ~IPackedLevelActorBuilder() {}
+	virtual FPackedLevelActorBuilderID GetID() const = 0;
+	virtual void GetPackClusters(FPackedLevelActorBuilderContext& InContext, AActor* InActor) const = 0;
+	virtual void PackActors(FPackedLevelActorBuilderContext& InBuilder, const FPackedLevelActorBuilderClusterID& InClusterID, const TArray<UActorComponent*>& InComponents) const = 0;
+};
+
+#endif
diff --git a/Engine/Source/Runtime/Engine/Public/PackedLevelActor/PackedLevelActor.h b/Engine/Source/Runtime/Engine/Public/PackedLevelActor/PackedLevelActor.h
--- a/Engine/Source/Runtime/Engine/Public/PackedLevelActor/PackedLevelActor.h
+++ b/Engine/Source/Runtime/Engine/Public/PackedLevelActor/PackedLevelActor.h
@@ -43,9 +43,12 @@ public:
 
 	ENGINE_API virtual TUniquePtr<class FWorldPartitionActorDesc> CreateClassActorDesc() const override;
 	
-	static ENGINE_API bool CreateOrUpdateBlueprint(ALevelInstance* InLevelInstance, TSoftObjectPtr<UBlueprint> InBlueprintAsset, bool bCheckoutAndSave = true, bool bPromptForSave = true);
-	static ENGINE_API bool CreateOrUpdateBlueprint(TSoftObjectPtr<UWorld> InWorldAsset, TSoftObjectPtr<UBlueprint> InBlueprintAsset, bool bCheckoutAndSave = true, bool bPromptForSave = true);
-	static ENGINE_API void UpdateBlueprint(UBlueprint* InBlueprint, bool bCheckoutAndSave = true);
+	UE_DEPRECATED(5.3, "Use FPackedLevelActorUtils::CreateOrUpdateBlueprint instead")
+	static ENGINE_API bool CreateOrUpdateBlueprint(ALevelInstance* InLevelInstance, TSoftObjectPtr<UBlueprint> InBlueprintAsset, bool bCheckoutAndSave = true, bool bPromptForSave = true) { return false;}
+	UE_DEPRECATED(5.3, "Use FPackedLevelActorUtils::CreateOrUpdateBlueprint instead")
+	static ENGINE_API bool CreateOrUpdateBlueprint(TSoftObjectPtr<UWorld> InWorldAsset, TSoftObjectPtr<UBlueprint> InBlueprintAsset, bool bCheckoutAndSave = true, bool bPromptForSave = true) { return false;}
+	UE_DEPRECATED(5.3, "Use FPackedLevelActorUtils::UpdateBlueprint instead")
+	static ENGINE_API void UpdateBlueprint(UBlueprint* InBlueprint, bool bCheckoutAndSave = true) {}
 
 	static ENGINE_API FName GetPackedComponentTag();
 
diff --git a/Engine/Source/Runtime/Engine/Public/PackedLevelActor/PackedLevelActorBuilder.h b/Engine/Source/Runtime/Engine/Public/PackedLevelActor/PackedLevelActorBuilder.h
new file mode 100644
--- /dev/null
+++ b/Engine/Source/Runtime/Engine/Public/PackedLevelActor/PackedLevelActorBuilder.h
@@ -0,0 +1,112 @@
+// Copyright Epic Games, Inc. All Rights Reserved.
+
+#pragma once
+
+#if WITH_EDITOR
+
+#include "CoreMinimal.h"
+#include "UObject/ObjectMacros.h"
+#include "PackedLevelActor/PackedLevelActorTypes.h"
+#include "PackedLevelActor/IPackedLevelActorBuilder.h"
+#include "UObject/SoftObjectPtr.h"
+#include "Containers/Set.h"
+
+class ILevelInstanceInterface;
+class APackedLevelActor;
+class ALevelInstance;
+class AActor;
+class UActorComponent;
+class UBlueprint;
+class FMessageLog;
+struct FWorldPartitionActorFilter;
+
+/**
+ * FPackedLevelActorBuilder handles packing of ALevelInstance actors into APackedLevelActor actors and Blueprints.
+ */
+class FPackedLevelActorBuilder
+{
+public:
+	ENGINE_API FPackedLevelActorBuilder();
+	
+	static ENGINE_API TSharedPtr<FPackedLevelActorBuilder> CreateDefaultBuilder();
+	
+	/* Packs InPackedLevelActor using InLevelInstanceToPack as its source level actor */
+	ENGINE_API bool PackActor(APackedLevelActor* InPackedLevelActor, ILevelInstanceInterface* InLevelInstanceToPack);
+	/* Packs InPackedLevelActor using itself as the source level actor */
+	ENGINE_API bool PackActor(APackedLevelActor* InPackedLevelActor);
+	/* Packs InPackedLevelActor using InWorldAsset as the source level */
+	ENGINE_API bool PackActor(APackedLevelActor* InPackedLevelActor, TSoftObjectPtr<UWorld> InWorldAsset);
+
+	/* Creates/Updates a APackedLevelActor Blueprint from InLevelInstance (will overwrite existing asset or show a dialog if InBlueprintAsset is null) */
+	ENGINE_API bool CreateOrUpdateBlueprint(ILevelInstanceInterface* InLevelInstance, TSoftObjectPtr<UBlueprint> InBlueprintAsset, bool bCheckoutAndSave = true, bool bPromptForSave = true);
+	/* Creates/Updates a APackedLeveInstance Blueprint from InWorldAsset (will overwrite existing asset or show a dialog if InBlueprintAsset is null) */
+	ENGINE_API bool CreateOrUpdateBlueprint(TSoftObjectPtr<UWorld> InWorldAsset, TSoftObjectPtr<UBlueprint> InBlueprintAsset, bool bCheckoutAndSave = true, bool bPromptForSave = true);
+	/* Update existing Blueprint */
+	ENGINE_API void UpdateBlueprint(UBlueprint* Blueprint, bool bCheckoutAndSave = true);
+
+	static ENGINE_API const FString& GetPackedBPPrefix();
+	/* Creates a new APackedLevelActor Blueprint using InPackagePath/InAssetName as hint for path. Prompts the user to input the final asset name. */
+	static ENGINE_API UBlueprint* CreatePackedLevelActorBlueprintWithDialog(TSoftObjectPtr<UBlueprint> InBlueprintAsset, TSoftObjectPtr<UWorld> InWorldAsset, bool bInCompile);
+	/* Creates a new APackedLevelActor Blueprint using InPackagePath/InAssetName */
+	static ENGINE_API UBlueprint* CreatePackedLevelActorBlueprint(TSoftObjectPtr<UBlueprint> InBlueprintAsset, TSoftObjectPtr<UWorld> InWorldAsset, bool bInCompile);
+	
+private:
+	ENGINE_API bool PackActor(FPackedLevelActorBuilderContext& InContext);
+
+	/* Create/Updates a APackedLevelActor Blueprint from InPackedActor (will overwrite existing asset or show a dialog if InBlueprintAsset is null) */
+	ENGINE_API bool CreateOrUpdateBlueprintFromPacked(APackedLevelActor* InPackedActor, TSoftObjectPtr<UBlueprint> InBlueprintAsset, bool bCheckoutAndSave, bool bPromptForSave);
+	/* Creates/Updates a APackedLevelActor Blueprint from InLevelInstance (will overwrite existing asset or show a dialog if InBlueprintAsset is null) */
+	ENGINE_API bool CreateOrUpdateBlueprintFromUnpacked(ILevelInstanceInterface* InLevelInstance, TSoftObjectPtr<UBlueprint> InBlueprintAsset, bool bCheckoutAndSave, bool bPromptForSave);
+	/* Creates and loads a ALevelInstance so it can be used for packing */
+	ENGINE_API ALevelInstance* CreateTransientLevelInstanceForPacking(TSoftObjectPtr<UWorld> InWorldAsset, const FVector& InLocation, const FRotator& InRotator, const FWorldPartitionActorFilter& InFilter);
+
+	FPackedLevelActorBuilder(const FPackedLevelActorBuilder&) = delete;
+	FPackedLevelActorBuilder& operator=(const FPackedLevelActorBuilder&) = delete;
+
+public:
+	friend class FPackedLevelActorBuilderContext;
+		
+private:
+	
+	TSet<UClass*> ClassDiscards;
+	TMap<FPackedLevelActorBuilderID, TUniquePtr<IPackedLevelActorBuilder>> Builders;
+};
+
+
+class FPackedLevelActorBuilderContext
+{
+public:
+	FPackedLevelActorBuilderContext(const FPackedLevelActorBuilder& InBuilder, APackedLevelActor* InPackedLevelActor, ILevelInstanceInterface* InLevelInstanceToPack) 
+		: Builders(InBuilder.Builders), ClassDiscards(InBuilder.ClassDiscards), PackedLevelActor(InPackedLevelActor), LevelInstanceToPack(InLevelInstanceToPack), RelativePivotTransform(FTransform::Identity) {}
+
+	/* Interface for IPackedLevelActorBuilder's to use */
+	void ClusterLevelActor(AActor* InLevelActor);
+	void FindOrAddCluster(FPackedLevelActorBuilderClusterID&& InClusterID, UActorComponent* InComponent = nullptr);
+	void DiscardActor(AActor* InActor);
+	void Report(FMessageLog& Log) const;
+
+	const TMap<FPackedLevelActorBuilderClusterID, TArray<UActorComponent*>>& GetClusters() const { return Clusters; }
+	
+	void SetRelativePivotTransform(const FTransform& InRelativePivotTransform) { RelativePivotTransform = InRelativePivotTransform; }
+	const FTransform GetRelativePivotTransform() const { return RelativePivotTransform; }
+
+	bool ShouldPackComponent(UActorComponent* InActorComponent) const;
+
+	APackedLevelActor* GetPackedLevelActor() const { return PackedLevelActor; }
+	ILevelInstanceInterface* GetLevelInstanceToPack() const { return LevelInstanceToPack; }
+private:
+	const TMap<FPackedLevelActorBuilderID, TUniquePtr<IPackedLevelActorBuilder>>& Builders;
+	const TSet<UClass*>& ClassDiscards;
+
+	APackedLevelActor* PackedLevelActor;
+	ILevelInstanceInterface* LevelInstanceToPack;
+	
+	TMap<FPackedLevelActorBuilderClusterID, TArray<UActorComponent*>> Clusters;
+
+	TMap<AActor*, TSet<UActorComponent*>> PerActorClusteredComponents;
+	TSet<AActor*> ActorDiscards;
+
+	FTransform RelativePivotTransform;
+};
+
+#endif
diff --git a/Engine/Source/Runtime/Engine/Public/PackedLevelActor/PackedLevelActorTypes.h b/Engine/Source/Runtime/Engine/Public/PackedLevelActor/PackedLevelActorTypes.h
new file mode 100644
--- /dev/null
+++ b/Engine/Source/Runtime/Engine/Public/PackedLevelActor/PackedLevelActorTypes.h
@@ -0,0 +1,92 @@
+// Copyright Epic Games, Inc. All Rights Reserved.
+
+#pragma once
+
+#if WITH_EDITOR
+
+#include "CoreMinimal.h"
+#include "UObject/ObjectMacros.h"
+#include "Misc/Crc.h"
+#include "Logging/LogMacros.h"
+
+using FPackedLevelActorBuilderID = uint32;
+
+class FPackedLevelActorBuilderCluster
+{
+	FPackedLevelActorBuilderID BuilderID;
+
+public:
+	FPackedLevelActorBuilderCluster(FPackedLevelActorBuilderID InBuilderID)
+		: BuilderID(InBuilderID) {}
+	virtual ~FPackedLevelActorBuilderCluster() {}
+
+	FPackedLevelActorBuilderID GetBuilderID() const { return BuilderID; }
+
+	virtual uint32 ComputeHash() const
+	{
+		return FCrc::TypeCrc32(BuilderID);
+	}
+
+	virtual bool Equals(const FPackedLevelActorBuilderCluster& Other) const
+	{
+		return BuilderID == Other.BuilderID;
+	}
+
+private:
+	FPackedLevelActorBuilderCluster(const FPackedLevelActorBuilderCluster&) = delete;
+	FPackedLevelActorBuilderCluster& operator=(const FPackedLevelActorBuilderCluster&) = delete;
+};
+
+class FPackedLevelActorBuilderClusterID
+{
+public:
+	static ENGINE_API FPackedLevelActorBuilderClusterID Invalid;
+
+	FPackedLevelActorBuilderClusterID() : Hash(0)
+	{
+	}
+
+	FPackedLevelActorBuilderClusterID(FPackedLevelActorBuilderClusterID&& Other)
+		: Hash(Other.Hash), Data(MoveTemp(Other.Data))
+	{
+	}
+
+	FPackedLevelActorBuilderClusterID(TUniquePtr<FPackedLevelActorBuilderCluster>&& InData)
+		: Data(MoveTemp(InData))
+	{
+		Hash = Data->ComputeHash();
+	}
+		
+	bool operator==(const FPackedLevelActorBuilderClusterID& Other) const
+	{
+		return (Hash == Other.Hash) && Data->Equals(*Other.Data);
+	}
+
+	bool operator!=(const FPackedLevelActorBuilderClusterID& Other) const
+	{
+		return !(*this == Other);
+	}
+
+	uint32 GetHash() const { return Hash; }
+
+	FPackedLevelActorBuilderID GetBuilderID() const { return Data->GetBuilderID(); }
+		
+	friend uint32 GetTypeHash(const FPackedLevelActorBuilderClusterID& ID)
+	{
+		return ID.GetHash();
+	}
+
+	FPackedLevelActorBuilderCluster* GetData() const
+	{
+		return Data.Get();
+	}
+
+private:
+	FPackedLevelActorBuilderClusterID(const FPackedLevelActorBuilderClusterID&) = delete;
+	FPackedLevelActorBuilderClusterID& operator=(const FPackedLevelActorBuilderClusterID&) = delete;
+
+	uint32 Hash;
+	TUniquePtr<FPackedLevelActorBuilderCluster> Data;
+};
+
+#endif
