Packed Level Actor: Remove unused param

#rnx

diff --git a/Engine/Source/Editor/LevelInstanceEditor/Private/PackedLevelActorUtils.cpp b/Engine/Source/Editor/LevelInstanceEditor/Private/PackedLevelActorUtils.cpp
--- a/Engine/Source/Editor/LevelInstanceEditor/Private/PackedLevelActorUtils.cpp
+++ b/Engine/Source/Editor/LevelInstanceEditor/Private/PackedLevelActorUtils.cpp
@@ -16,7 +16,7 @@
 
 #define LOCTEXT_NAMESPACE "FPackedLevelActorUtils"
 
-void FPackedLevelActorUtils::GetPackedBlueprintsForWorldAsset(const TSoftObjectPtr<UWorld>& InWorldAsset, TSet<TSoftObjectPtr<UBlueprint>>& OutPackedBlueprintAssets, bool bInLoadedOnly, bool bInDiskOnly)
+void FPackedLevelActorUtils::GetPackedBlueprintsForWorldAsset(const TSoftObjectPtr<UWorld>& InWorldAsset, TSet<TSoftObjectPtr<UBlueprint>>& OutPackedBlueprintAssets, bool bInLoadedOnly)
 {
 	FAssetRegistryModule& AssetRegistryModule = FModuleManager::LoadModuleChecked<FAssetRegistryModule>(FName("AssetRegistry"));
 	IAssetRegistry& AssetRegistry = AssetRegistryModule.Get();
diff --git a/Engine/Source/Editor/LevelInstanceEditor/Public/PackedLevelActorUtils.h b/Engine/Source/Editor/LevelInstanceEditor/Public/PackedLevelActorUtils.h
--- a/Engine/Source/Editor/LevelInstanceEditor/Public/PackedLevelActorUtils.h
+++ b/Engine/Source/Editor/LevelInstanceEditor/Public/PackedLevelActorUtils.h
@@ -12,7 +12,7 @@ class ILevelInstanceInterface;
 class FPackedLevelActorUtils
 {
 public:
-	static void GetPackedBlueprintsForWorldAsset(const TSoftObjectPtr<UWorld>& InWorldAsset, TSet<TSoftObjectPtr<UBlueprint>>& OutPackedBlueprintAssets, bool bInLoadedOnly = false, bool bInDiskOnly = false);
+	static void GetPackedBlueprintsForWorldAsset(const TSoftObjectPtr<UWorld>& InWorldAsset, TSet<TSoftObjectPtr<UBlueprint>>& OutPackedBlueprintAssets, bool bInLoadedOnly = false);
 
 	static bool CanPack();
 
