Level Instance: allow replacing actors to level instance class.

diff --git a/Engine/Source/Editor/LevelInstanceEditor/Private/LevelInstanceActorFactory.cpp b/Engine/Source/Editor/LevelInstanceEditor/Private/LevelInstanceActorFactory.cpp
--- a/Engine/Source/Editor/LevelInstanceEditor/Private/LevelInstanceActorFactory.cpp
+++ b/Engine/Source/Editor/LevelInstanceEditor/Private/LevelInstanceActorFactory.cpp
@@ -26,12 +26,30 @@ bool ULevelInstanceActorFactory::CanCreateActorFrom(const FAssetData& AssetData,
 		return false;
 	}
 
-	// If asset is valid it needs to be of type: UWorld
-	if (AssetData.IsValid() && !AssetData.IsInstanceOf(UWorld::StaticClass()))
+	if (AssetData.IsValid())
 	{
-		OutErrorMsg = NSLOCTEXT("LevelInstanceActorFactory", "NoWorld", "A valid world must be specified.");
-		return false;
+		// Only allow creating level instance actors from actor classes
+		if (AssetData.IsInstanceOf<UClass>())
+		{
+			if (UClass* Class = Cast<UClass>(AssetData.GetAsset()))
+			{
+				if (!Class->IsChildOf<AActor>() || !Cast<ILevelInstanceInterface>(Class->GetDefaultObject()))
+				{
+					OutErrorMsg = NSLOCTEXT("LevelInstanceActorFactory", "InvalidClass", "A valid actor class must be specified.");
+					return false;
+				}
+			
+				return true;
+			}
+		}
+
+		// Only allow creating level instance actors from world assets
+		if (!AssetData.IsInstanceOf<UWorld>())
+		{
+			OutErrorMsg = NSLOCTEXT("LevelInstanceActorFactory", "NoWorld", "A valid world must be specified.");
+			return false;
+		}
 	}
 
 	return true;
-}
+}
\ No newline at end of file
