LevelInstance
- Remove ShouldExport restriction (was added recently while enabling sub-selection but it doesn't cause issues) which was causing regression in Break Level Instance
- Prevent some operations (Break, Move, Create) when some Level Instances are selected that have selection parents (non-root selections)

#rnx

diff --git a/Engine/Source/Editor/LevelInstanceEditor/Private/LevelInstanceEditorModule.cpp b/Engine/Source/Editor/LevelInstanceEditor/Private/LevelInstanceEditorModule.cpp
--- a/Engine/Source/Editor/LevelInstanceEditor/Private/LevelInstanceEditorModule.cpp
+++ b/Engine/Source/Editor/LevelInstanceEditor/Private/LevelInstanceEditorModule.cpp
@@ -285,6 +285,22 @@ namespace LevelInstanceMenuUtils
 		return ALevelInstance::StaticClass();
 	}
 
+	bool AreAllSelectedLevelInstancesRootSelections()
+	{
+		for (FSelectionIterator It(GEditor->GetSelectedActorIterator()); It; ++It)
+		{
+			if (ILevelInstanceInterface* LevelInstance = Cast<ILevelInstanceInterface>(*It))
+			{
+				if (CastChecked<AActor>(*It)->GetSelectionParent() != nullptr)
+				{
+					return false;
+				}
+			}
+		}
+
+		return true;
+	}
+		
 	void CreateLevelInstanceFromSelection(ULevelInstanceSubsystem* LevelInstanceSubsystem, ELevelInstanceCreationType CreationType)
 	{
 		TArray<AActor*> ActorsToMove;
@@ -872,6 +888,9 @@ void FLevelInstanceEditorModule::ExtendContextMenu()
 			return;
 		}
 
+		// Some actions aren't allowed on non root selection Level Instances (Readonly Level Instances)
+		const bool bAreAllSelectedLevelInstancesRootSelections = LevelInstanceMenuUtils::AreAllSelectedLevelInstancesRootSelections();
+
 		if (ULevelEditorContextMenuContext* LevelEditorMenuContext = ToolMenu->Context.FindContext<ULevelEditorContextMenuContext>())
 		{
 			// Use the actor under the cursor if available (e.g. right-click menu).
@@ -884,17 +903,28 @@ void FLevelInstanceEditorModule::ExtendContextMenu()
 
 			if (ContextActor)
 			{
+				// Allow Edit/Commmit on non root selected Level Instance
 				LevelInstanceMenuUtils::CreateEditMenu(ToolMenu, ContextActor);
 				LevelInstanceMenuUtils::CreateCommitDiscardMenu(ToolMenu, ContextActor);
-				LevelInstanceMenuUtils::CreatePackedBlueprintMenu(ToolMenu, ContextActor);
-				LevelInstanceMenuUtils::CreateSetCurrentMenu(ToolMenu, ContextActor);
+				
+				if (bAreAllSelectedLevelInstancesRootSelections)
+				{
+					LevelInstanceMenuUtils::CreatePackedBlueprintMenu(ToolMenu, ContextActor);
+					LevelInstanceMenuUtils::CreateSetCurrentMenu(ToolMenu, ContextActor);
+				}
 			}
 
-			LevelInstanceMenuUtils::CreateMoveSelectionToMenu(ToolMenu);
+			if (bAreAllSelectedLevelInstancesRootSelections)
+			{
+				LevelInstanceMenuUtils::CreateMoveSelectionToMenu(ToolMenu);
+			}
 		}
 
-		LevelInstanceMenuUtils::CreateBreakMenu(ToolMenu);
-		LevelInstanceMenuUtils::CreateCreateMenu(ToolMenu);
+		if (bAreAllSelectedLevelInstancesRootSelections)
+		{
+			LevelInstanceMenuUtils::CreateBreakMenu(ToolMenu);
+			LevelInstanceMenuUtils::CreateCreateMenu(ToolMenu);
+		}
 	};
 
 	if (UToolMenu* ToolMenu = UToolMenus::Get()->ExtendMenu("LevelEditor.ActorContextMenu.LevelSubMenu"))
diff --git a/Engine/Source/Runtime/Engine/Private/LevelInstance/LevelInstanceActorImpl.cpp b/Engine/Source/Runtime/Engine/Private/LevelInstance/LevelInstanceActorImpl.cpp
--- a/Engine/Source/Runtime/Engine/Private/LevelInstance/LevelInstanceActorImpl.cpp
+++ b/Engine/Source/Runtime/Engine/Private/LevelInstance/LevelInstanceActorImpl.cpp
@@ -325,7 +325,7 @@ bool FLevelInstanceActorImpl::IsLockedActor() const
 
 bool FLevelInstanceActorImpl::ShouldExport() const
 {
-	return !IsLockedActor();
+	return true;
 }
 
 bool FLevelInstanceActorImpl::IsUserManaged() const
