Level Instance Details: use ILevelInstanceInterface so details can be used on different impl.

diff --git a/Engine/Source/Editor/LevelInstanceEditor/Private/LevelInstanceActorDetails.cpp b/Engine/Source/Editor/LevelInstanceEditor/Private/LevelInstanceActorDetails.cpp
--- a/Engine/Source/Editor/LevelInstanceEditor/Private/LevelInstanceActorDetails.cpp
+++ b/Engine/Source/Editor/LevelInstanceEditor/Private/LevelInstanceActorDetails.cpp
@@ -9,27 +9,27 @@
 #include "ScopedTransaction.h"
 
 #include "Engine/World.h"
-#include "LevelInstance/LevelInstanceActor.h"
+#include "LevelInstance/LevelInstanceInterface.h"
 
 #define LOCTEXT_NAMESPACE "FLevelInstanceActorDetails"
 
 namespace LevelInstanceActorDetailsCallbacks
 {
-	static bool IsEditCommitButtonEnabled(TWeakObjectPtr<ALevelInstance> LevelInstanceActorPtr)
+	static bool IsEditCommitButtonEnabled(TWeakInterfacePtr<ILevelInstanceInterface> LevelInstancePtr)
 	{
-		if (ALevelInstance* LevelInstanceActor = LevelInstanceActorPtr.Get())
+		if (ILevelInstanceInterface* LevelInstance = LevelInstancePtr.Get())
 		{
-			return LevelInstanceActor->CanEnterEdit() || LevelInstanceActor->CanExitEdit();
+			return LevelInstance->CanEnterEdit() || LevelInstance->CanExitEdit();
 		}
 
 		return false;
 	}
 
-	static FText GetEditCommitButtonText(TWeakObjectPtr<ALevelInstance> LevelInstanceActorPtr)
+	static FText GetEditCommitButtonText(TWeakInterfacePtr<ILevelInstanceInterface> LevelInstancePtr)
 	{
-		if (ALevelInstance* LevelInstanceActor = LevelInstanceActorPtr.Get())
+		if (ILevelInstanceInterface* LevelInstance = LevelInstancePtr.Get())
 		{
-			if (LevelInstanceActor->CanExitEdit())
+			if (LevelInstance->CanExitEdit())
 			{
 				return LOCTEXT("CommitChanges", "Commit Changes");
 			}
@@ -38,43 +38,43 @@ namespace LevelInstanceActorDetailsCallbacks
 		return LOCTEXT("Edit", "Edit");
 	}
 
-	static FText GetEditCommitReasonText(TWeakObjectPtr<ALevelInstance> LevelInstanceActorPtr)
+	static FText GetEditCommitReasonText(TWeakInterfacePtr<ILevelInstanceInterface> LevelInstancePtr)
 	{
 		FText Reason;
-		if (ALevelInstance* LevelInstanceActor = LevelInstanceActorPtr.Get())
+		if (ILevelInstanceInterface* LevelInstance = LevelInstancePtr.Get())
 		{
-			if (!LevelInstanceActor->IsEditing())
+			if (!LevelInstance->IsEditing())
 			{
-				LevelInstanceActor->CanEnterEdit(&Reason);
+				LevelInstance->CanEnterEdit(&Reason);
 				return Reason;
 			}
 
-			LevelInstanceActor->CanExitEdit(/*bDiscardEdits=*/false, &Reason);
+			LevelInstance->CanExitEdit(/*bDiscardEdits=*/false, &Reason);
 		}
 		return Reason;
 	}
 
-	static EVisibility GetEditCommitReasonVisibility(TWeakObjectPtr<ALevelInstance> LevelInstanceActorPtr)
+	static EVisibility GetEditCommitReasonVisibility(TWeakInterfacePtr<ILevelInstanceInterface> LevelInstancePtr)
 	{
-		if (ALevelInstance* LevelInstanceActor = LevelInstanceActorPtr.Get())
+		if (ILevelInstanceInterface* LevelInstance = LevelInstancePtr.Get())
 		{
-			return IsEditCommitButtonEnabled(LevelInstanceActor) ? EVisibility::Collapsed : EVisibility::Visible;
+			return IsEditCommitButtonEnabled(LevelInstance) ? EVisibility::Collapsed : EVisibility::Visible;
 		}
 
 		return EVisibility::Collapsed;
 	}
 
-	static FReply OnEditCommitButtonClicked(TWeakObjectPtr<ALevelInstance> LevelInstanceActorPtr)
+	static FReply OnEditCommitButtonClicked(TWeakInterfacePtr<ILevelInstanceInterface> LevelInstancePtr)
 	{
-		if (ALevelInstance* LevelInstanceActor = LevelInstanceActorPtr.Get())
+		if (ILevelInstanceInterface* LevelInstance = LevelInstancePtr.Get())
 		{
-			if (LevelInstanceActor->CanExitEdit())
+			if (LevelInstance->CanExitEdit())
 			{
-				LevelInstanceActor->ExitEdit();
+				LevelInstance->ExitEdit();
 			}
-			else if (LevelInstanceActor->CanEnterEdit())
+			else if (LevelInstance->CanEnterEdit())
 			{
-				LevelInstanceActor->EnterEdit();
+				LevelInstance->EnterEdit();
 			}
 		}
 		return FReply::Handled();
@@ -100,8 +100,8 @@ void FLevelInstanceActorDetails::CustomizeDetails(IDetailLayoutBuilder& DetailBu
 		return;
 	}
 
-	TWeakObjectPtr<ALevelInstance> EditingObject = Cast<ALevelInstance>(EditingObjects[0].Get());
-	UWorld* World = EditingObject->GetWorld();
+	TWeakInterfacePtr<ILevelInstanceInterface> LevelInstance = Cast<ILevelInstanceInterface>(EditingObjects[0].Get());
+	UWorld* World = LevelInstance.GetObject()->GetWorld();
 
 	if (!World)
 	{
@@ -121,10 +121,10 @@ void FLevelInstanceActorDetails::CustomizeDetails(IDetailLayoutBuilder& DetailBu
 			.FillWidth(1)
 			[
 				SNew(SMultiLineEditableTextBox)
-				.Visibility_Static(&LevelInstanceActorDetailsCallbacks::GetEditCommitReasonVisibility, EditingObject)
+				.Visibility_Static(&LevelInstanceActorDetailsCallbacks::GetEditCommitReasonVisibility, LevelInstance)
 				.Font(DetailBuilder.GetDetailFontBold())
 				.BackgroundColor(TAttribute<FSlateColor>::Create(TAttribute<FSlateColor>::FGetter::CreateLambda([]() { return FAppStyle::GetColor("ErrorReporting.WarningBackgroundColor"); })))
-				.Text_Static(&LevelInstanceActorDetailsCallbacks::GetEditCommitReasonText, EditingObject)
+				.Text_Static(&LevelInstanceActorDetailsCallbacks::GetEditCommitReasonText, LevelInstance)
 				.AutoWrapText(true)
 				.IsReadOnly(true)
 			]
@@ -137,10 +137,10 @@ void FLevelInstanceActorDetails::CustomizeDetails(IDetailLayoutBuilder& DetailBu
 			.FillWidth(1)
 			[
 				SNew(SButton)
-				.IsEnabled_Static(&LevelInstanceActorDetailsCallbacks::IsEditCommitButtonEnabled, EditingObject)
-				.Text_Static(&LevelInstanceActorDetailsCallbacks::GetEditCommitButtonText, EditingObject)
+				.IsEnabled_Static(&LevelInstanceActorDetailsCallbacks::IsEditCommitButtonEnabled, LevelInstance)
+				.Text_Static(&LevelInstanceActorDetailsCallbacks::GetEditCommitButtonText, LevelInstance)
 				.HAlign(HAlign_Center)
-				.OnClicked_Static(&LevelInstanceActorDetailsCallbacks::OnEditCommitButtonClicked, EditingObject)
+				.OnClicked_Static(&LevelInstanceActorDetailsCallbacks::OnEditCommitButtonClicked, LevelInstance)
 			]
 		]
 	];
