[Experimental] Level Instance Property Override
- World Partition support only
- Property Overrides are saved as a sub-object (ULevelInstancePropertyOverrideAsset) of the LevelInstance Actor (same package)
- Support for Overrides to be stored at any level of the LevelInstance hierarchy (based on root edit context)
- Support for Override ActorDescs, overriden actors have override actor descs through ULevelInstanceContainerInstance/ULevelInstancePropertyOverrideContainer meaning overriden ActorDescs which will be taken into account for streaming generation
- Override ActorDescs are saved as part of a FLevelInstancePropertyOverrideDesc
- Add support for ActorDescs to be serialized with a base actor desc instead of a class desc
- Add Level Instance Property Override Experimental setting
- Remove Level Instance / Packed Level Actor Experimental setting (no longer experimental)

#rnx

diff --git a/Engine/Source/Editor/LevelInstanceEditor/LevelInstanceEditor.Build.cs b/Engine/Source/Editor/LevelInstanceEditor/LevelInstanceEditor.Build.cs
--- a/Engine/Source/Editor/LevelInstanceEditor/LevelInstanceEditor.Build.cs
+++ b/Engine/Source/Editor/LevelInstanceEditor/LevelInstanceEditor.Build.cs
@@ -34,7 +34,8 @@ public class LevelInstanceEditor : ModuleRules
 				"DeveloperSettings",
 				"SceneOutliner",
 				"WorldPartitionEditor",
-				"Kismet"
+				"Kismet",
+				"PropertyEditor"
 			}
 		);
 
diff --git a/Engine/Source/Editor/LevelInstanceEditor/Private/LevelInstanceActorDetails.cpp b/Engine/Source/Editor/LevelInstanceEditor/Private/LevelInstanceActorDetails.cpp
--- a/Engine/Source/Editor/LevelInstanceEditor/Private/LevelInstanceActorDetails.cpp
+++ b/Engine/Source/Editor/LevelInstanceEditor/Private/LevelInstanceActorDetails.cpp
@@ -109,6 +109,11 @@ void FLevelInstanceActorDetails::CustomizeDetails(IDetailLayoutBuilder& DetailBu
 		return;
 	}
 
+	if (ULevelInstancePropertyOverrideAsset* PropertyOverride = LevelInstance->GetPropertyOverrideAsset())
+	{
+		DetailBuilder.HideProperty("WorldAsset");
+	}
+
 	IDetailCategoryBuilder& LevelInstanceEditingCategory = DetailBuilder.EditCategory("Level Edit", FText::GetEmpty(), ECategoryPriority::Transform);
 
 	LevelInstanceEditingCategory.AddCustomRow(FText::GetEmpty())
diff --git a/Engine/Source/Editor/LevelInstanceEditor/Private/LevelInstanceActorFactory.cpp b/Engine/Source/Editor/LevelInstanceEditor/Private/LevelInstanceActorFactory.cpp
--- a/Engine/Source/Editor/LevelInstanceEditor/Private/LevelInstanceActorFactory.cpp
+++ b/Engine/Source/Editor/LevelInstanceEditor/Private/LevelInstanceActorFactory.cpp
@@ -3,7 +3,7 @@
 #include "LevelInstanceActorFactory.h"
 #include "LevelInstanceEditorSettings.h"
 #include "LevelInstance/LevelInstanceInterface.h"
-#include "Settings/EditorExperimentalSettings.h"
+#include "LevelInstance/LevelInstanceSettings.h"
 
 ULevelInstanceActorFactory::ULevelInstanceActorFactory(const FObjectInitializer& ObjectInitializer)
 	: Super(ObjectInitializer)
@@ -20,9 +20,9 @@ void ULevelInstanceActorFactory::PostSpawnActor(UObject* Asset, AActor* NewActor
 
 bool ULevelInstanceActorFactory::CanCreateActorFrom(const FAssetData& AssetData, FText& OutErrorMsg)
 {
-	if (!GetDefault<UEditorExperimentalSettings>()->bLevelInstance)
+	if (GetDefault<ULevelInstanceSettings>()->IsLevelInstanceDisabled())
 	{
-		OutErrorMsg = NSLOCTEXT("LevelInstanceActorFactory", "ExperimentalSettings", "Level Instance must be enabled in experimental settings.");
+		OutErrorMsg = NSLOCTEXT("LevelInstanceActorFactory", "LevelInstanceDisabled", "Level Instance support is disabled.");
 		return false;
 	}
 
diff --git a/Engine/Source/Editor/LevelInstanceEditor/Private/LevelInstanceEditorMode.cpp b/Engine/Source/Editor/LevelInstanceEditor/Private/LevelInstanceEditorMode.cpp
--- a/Engine/Source/Editor/LevelInstanceEditor/Private/LevelInstanceEditorMode.cpp
+++ b/Engine/Source/Editor/LevelInstanceEditor/Private/LevelInstanceEditorMode.cpp
@@ -139,31 +139,50 @@ bool ULevelInstanceEditorMode::IsSelectionDisallowed(AActor* InActor, bool bInSe
 	if (bRestrict)
 	{
 		check(World);
-		if (ILevelInstanceInterface* LevelInstance = Cast<ILevelInstanceInterface>(InActor))
+		if (ULevelInstanceSubsystem* LevelInstanceSubsystem = UWorld::GetSubsystem<ULevelInstanceSubsystem>(World))
 		{
-			if (LevelInstance->IsEditing())
+			const ILevelInstanceInterface* PropertyOverrideLevelInstance = LevelInstanceSubsystem->GetEditingPropertyOverridesLevelInstance();
+			const ILevelInstanceInterface* EditLevelInstance = LevelInstanceSubsystem->GetEditingLevelInstance();
+
+			if (ILevelInstanceInterface* LevelInstance = Cast<ILevelInstanceInterface>(InActor))
 			{
-				return false;
+				// If Actor is itself a Level Instance and is one of the edits, allow selection
+				if (LevelInstance == PropertyOverrideLevelInstance || LevelInstance == EditLevelInstance)
+				{
+					return false;
+				}
 			}
-		}
 
-		if (ULevelInstanceSubsystem* LevelInstanceSubsystem = World->GetSubsystem<ULevelInstanceSubsystem>())
-		{
-			ILevelInstanceInterface* EditingLevelInstance = LevelInstanceSubsystem->GetEditingLevelInstance();
-			ILevelInstanceInterface* LevelInstance = LevelInstanceSubsystem->GetParentLevelInstance(InActor);
-			// Allow selection on actors that are part of the currently edited Level Instance hierarchy because AActor::GetRootSelectionParent() will eventually
-			// Bubble up the selection to its parent.
-			while (LevelInstance != nullptr)
+			const ILevelInstanceInterface* ParentLevelInstance = LevelInstanceSubsystem->GetParentLevelInstance(InActor);
+						
+			auto IsAncestorOrSelf = [LevelInstanceSubsystem](const ILevelInstanceInterface* LevelInstance, const ILevelInstanceInterface* Ancestor)
 			{
-				if (LevelInstance == EditingLevelInstance)
+				while (LevelInstance != nullptr)
 				{
-					return false;
+					if (LevelInstance == Ancestor)
+					{
+						return true;
+					}
+
+					LevelInstance = LevelInstanceSubsystem->GetParentLevelInstance(CastChecked<AActor>(LevelInstance));
 				}
 
-				LevelInstance = LevelInstanceSubsystem->GetParentLevelInstance(CastChecked<AActor>(LevelInstance));
+				return false;
+			};
+
+			// If we have a PropertyOverride Edit in progress, actor can be selected if it is part of the PropertyOverrides hierarchy
+			if (PropertyOverrideLevelInstance)
+			{
+				return !IsAncestorOrSelf(ParentLevelInstance, PropertyOverrideLevelInstance);
 			}
 
-			return EditingLevelInstance != nullptr;
+			// If we have a Edit in progress, actor can be selected if it is part of the Edit hierarchy
+			if (EditLevelInstance)
+			{
+				return !IsAncestorOrSelf(ParentLevelInstance, EditLevelInstance);
+			}
+
+			return false;
 		}
 	}
 
@@ -200,7 +219,12 @@ bool ULevelInstanceEditorMode::IsContextRestrictedForWorld(UWorld* InWorld) cons
 {
 	if (ULevelInstanceSubsystem* LevelInstanceSubsystem = InWorld? InWorld->GetSubsystem<ULevelInstanceSubsystem>() : nullptr)
 	{
-		if (ILevelInstanceInterface* EditingLevelInstance = LevelInstanceSubsystem->GetEditingLevelInstance())
+		if (ILevelInstanceInterface* EditingPropertyOverrides = LevelInstanceSubsystem->GetEditingPropertyOverridesLevelInstance())
+		{
+			// Always restrict outside selection while editing property overrides
+			return true;
+		}
+		else if (ILevelInstanceInterface* EditingLevelInstance = LevelInstanceSubsystem->GetEditingLevelInstance())
 		{
 			return bContextRestriction && LevelInstanceSubsystem->GetLevelInstanceLevel(EditingLevelInstance) == InWorld->GetCurrentLevel();
 		}
diff --git a/Engine/Source/Editor/LevelInstanceEditor/Private/LevelInstanceEditorModeToolkit.cpp b/Engine/Source/Editor/LevelInstanceEditor/Private/LevelInstanceEditorModeToolkit.cpp
--- a/Engine/Source/Editor/LevelInstanceEditor/Private/LevelInstanceEditorModeToolkit.cpp
+++ b/Engine/Source/Editor/LevelInstanceEditor/Private/LevelInstanceEditorModeToolkit.cpp
@@ -33,64 +33,62 @@ void FLevelInstanceEditorModeToolkit::Init(const TSharedPtr<IToolkitHost>& InitT
 {
 	FModeToolkit::Init(InitToolkitHost, InOwningMode);
 
-	const ULevelInstanceSubsystem* LevelInstanceSubsystem = UWorld::GetSubsystem<ULevelInstanceSubsystem>(InitToolkitHost->GetWorld());
+	ULevelInstanceSubsystem* LevelInstanceSubsystem = UWorld::GetSubsystem<ULevelInstanceSubsystem>(InitToolkitHost->GetWorld());
 	check(LevelInstanceSubsystem);
 
-	const ILevelInstanceInterface* EditingLevelInstance = LevelInstanceSubsystem->GetEditingLevelInstance();
-	check(EditingLevelInstance);
-	const TSoftObjectPtr<UWorld> WorldAsset = EditingLevelInstance->GetWorldAsset();
-	const FString WorldAssetName = WorldAsset.GetAssetName();
-	FText DisplayText = FText::FromString(WorldAsset.GetAssetName());
-
-	ILevelInstanceEditorModule& EditorModule = FModuleManager::GetModuleChecked<ILevelInstanceEditorModule>("LevelInstanceEditor");
-	
 	// ViewportOverlay
 	SAssignNew(ViewportOverlayWidget, SHorizontalBox)
-		+SHorizontalBox::Slot()
-		.HAlign(HAlign_Center)
-		.VAlign(VAlign_Bottom)
-		.Padding(FMargin(0.0f, 0.0f, 0.f, 15.f))
+	+SHorizontalBox::Slot()
+	.HAlign(HAlign_Center)
+	.VAlign(VAlign_Bottom)
+	.Padding(FMargin(0.0f, 0.0f, 0.f, 15.f))
+	[
+		SNew(SBorder)
+		.BorderImage(FAppStyle::Get().GetBrush("EditorViewport.OverlayBrush"))
+		.Padding(8.f)
 		[
-			SNew(SBorder)
-			.BorderImage(FAppStyle::Get().GetBrush("EditorViewport.OverlayBrush"))
-			.Padding(8.f)
+			SNew(SHorizontalBox)
+			+SHorizontalBox::Slot()
+			.AutoWidth()
+			.VAlign(VAlign_Center)
+			.Padding(FMargin(0.f, 0.f, 8.f, 0.f))
 			[
-				SNew(SHorizontalBox)
-				+SHorizontalBox::Slot()
-				.AutoWidth()
-				.VAlign(VAlign_Center)
-				.Padding(FMargin(0.f, 0.f, 8.f, 0.f))
-				[
-					SNew(SImage)
-					.Image(FSlateIconFinder::FindIconBrushForClass(ALevelInstance::StaticClass()))
-				]
-				+SHorizontalBox::Slot()
-				.AutoWidth()
-				.VAlign(VAlign_Center)
-				.Padding(FMargin(0.f, 0.f, 8.f, 0.f))
-				[
-					SNew(STextBlock)
-					.Text(DisplayText)
-				]
-				+SHorizontalBox::Slot()
-				.AutoWidth()
-				.Padding(FMargin(2.0, 0.f, 0.f, 0.f))
-				[
-					SNew(SButton)
-					.ButtonStyle(FAppStyle::Get(), "PrimaryButton")
-					.TextStyle(FAppStyle::Get(), "DialogButtonText")
-					.Text(LOCTEXT("ExitEdit", "Exit"))
-					.ToolTipText(LOCTEXT("ExitTooltip", "Exit Level Instance Edit"))
-					.HAlign(HAlign_Center)
-					.OnClicked_Lambda([&EditorModule]() 
-					{ 
-						EditorModule.BroadcastTryExitEditorMode(); 
-						return FReply::Handled(); 
-					})
-				]
+				SNew(SImage)
+				.Image(FSlateIconFinder::FindIconBrushForClass(ALevelInstance::StaticClass()))
 			]
-		];
-
+			+SHorizontalBox::Slot()
+			.AutoWidth()
+			.VAlign(VAlign_Center)
+			.Padding(FMargin(0.f, 0.f, 8.f, 0.f))
+			[
+				SNew(STextBlock)
+				.Text_Lambda([LevelInstanceSubsystem]()
+				{
+					return LevelInstanceSubsystem->GetToolKitDisplayText();
+				})
+			]
+			+SHorizontalBox::Slot()
+			.AutoWidth()
+			.Padding(FMargin(2.0, 0.f, 0.f, 0.f))
+			[
+				SNew(SButton)
+				.ButtonStyle(FAppStyle::Get(), "PrimaryButton")
+				.TextStyle(FAppStyle::Get(), "DialogButtonText")
+				.Text(LOCTEXT("ExitEdit", "Exit"))
+				.ToolTipText_Lambda([LevelInstanceSubsystem]()
+				{
+					return LevelInstanceSubsystem->GetToolKitExitToolTip();
+				})
+				.HAlign(HAlign_Center)
+				.OnClicked_Lambda([LevelInstanceSubsystem]() 
+				{ 
+					LevelInstanceSubsystem->ToolKitExit();
+					return FReply::Handled(); 
+				})
+			]
+		]
+	];
+	
 	GetToolkitHost()->AddViewportOverlayWidget(ViewportOverlayWidget.ToSharedRef());
 }
 
diff --git a/Engine/Source/Editor/LevelInstanceEditor/Private/LevelInstanceEditorModule.cpp b/Engine/Source/Editor/LevelInstanceEditor/Private/LevelInstanceEditorModule.cpp
--- a/Engine/Source/Editor/LevelInstanceEditor/Private/LevelInstanceEditorModule.cpp
+++ b/Engine/Source/Editor/LevelInstanceEditor/Private/LevelInstanceEditorModule.cpp
@@ -7,6 +7,7 @@
 #include "LevelInstance/LevelInstanceSubsystem.h"
 #include "LevelInstance/LevelInstanceInterface.h"
 #include "LevelInstance/LevelInstanceActor.h"
+#include "LevelInstance/LevelInstanceSettings.h"
 #include "PackedLevelActor/PackedLevelActor.h"
 #include "PackedLevelActor/PackedLevelActorBuilder.h"
 #include "LevelInstanceEditorSettings.h"
@@ -45,6 +46,7 @@
 #include "Settings/EditorExperimentalSettings.h"
 #include "WorldPartition/WorldPartitionConverter.h"
 #include "WorldPartition/WorldPartitionActorLoaderInterface.h"
+#include "ScopedTransaction.h"
 
 IMPLEMENT_MODULE( FLevelInstanceEditorModule, LevelInstanceEditor );
 
@@ -52,22 +54,9 @@ IMPLEMENT_MODULE( FLevelInstanceEditorModule, LevelInstanceEditor );
 
 DEFINE_LOG_CATEGORY_STATIC(LogLevelInstanceEditor, Log, All);
 
-namespace LevelInstanceMenuUtils
+struct FLevelInstanceMenuUtils
 {
-	bool IsExperimentalSettingEnabled(ILevelInstanceInterface* LevelInstance)
-	{
-		if (APackedLevelActor* Actor = Cast<APackedLevelActor>(LevelInstance))
-		{
-			if (!GetDefault<UEditorExperimentalSettings>()->bPackedLevelActor)
-			{
-				return false;
-			}
-		}
-
-		return GetDefault<UEditorExperimentalSettings>()->bLevelInstance;
-	}
-
-	FToolMenuSection& CreateLevelSection(UToolMenu* Menu)
+	static FToolMenuSection& CreateLevelSection(UToolMenu* Menu)
 	{
 		const FName LevelSectionName = TEXT("Level");
 		FToolMenuSection* SectionPtr = Menu->FindSection(LevelSectionName);
@@ -79,7 +68,7 @@ namespace LevelInstanceMenuUtils
 		return Section;
 	}
 
-	void CreateEditMenuEntry(FToolMenuSection& Section, ILevelInstanceInterface* LevelInstance, AActor* ContextActor, bool bSingleEntry)
+	static void CreateEditMenuEntry(FToolMenuSection& Section, ILevelInstanceInterface* LevelInstance, AActor* ContextActor, bool bSingleEntry)
 	{
 		FToolUIAction LevelInstanceEditAction;
 		FText EntryDesc;
@@ -103,7 +92,7 @@ namespace LevelInstanceMenuUtils
 		Section.AddMenuEntry(NAME_None, EntryLabel, EntryDesc, FSlateIcon(), LevelInstanceEditAction);
 	}
 
-	void CreateEditSubMenu(UToolMenu* Menu, TArray<ILevelInstanceInterface*> LevelInstanceHierarchy, AActor* ContextActor)
+	static void CreateEditSubMenu(UToolMenu* Menu, TArray<ILevelInstanceInterface*> LevelInstanceHierarchy, AActor* ContextActor)
 	{
 		FToolMenuSection& Section = Menu->AddSection(NAME_None, LOCTEXT("LevelInstanceContextEditSection", "Context"));
 		for (ILevelInstanceInterface* LevelInstance : LevelInstanceHierarchy)
@@ -111,33 +100,53 @@ namespace LevelInstanceMenuUtils
 			CreateEditMenuEntry(Section, LevelInstance, ContextActor, false);
 		}
 	}
-		
-	void MoveSelectionToLevelInstance(ILevelInstanceInterface* DestinationLevelInstance)
+
+	static void CreateEditPropertyOverridesMenuEntry(FToolMenuSection& Section, ILevelInstanceInterface* LevelInstance, AActor* ContextActor, bool bSingleEntry)
 	{
-		TArray<AActor*> ActorsToMove;
-		ActorsToMove.Reserve(GEditor->GetSelectedActorCount());
-		for (FSelectionIterator It(GEditor->GetSelectedActorIterator()); It; ++It)
+		FToolUIAction LevelInstanceEditAction;
+		FText EntryDesc;
+		AActor* LevelInstanceActor = CastChecked<AActor>(LevelInstance);
+		const bool bCanEdit = LevelInstance->CanEnterEditPropertyOverrides(&EntryDesc);
+
+		LevelInstanceEditAction.ExecuteAction.BindLambda([LevelInstance, ContextActor](const FToolMenuContext&)
 		{
-			if (AActor* Actor = Cast<AActor>(*It))
-			{
-				ActorsToMove.Add(Actor);
-			}
+			LevelInstance->EnterEditPropertyOverrides(ContextActor);
+		});
+		LevelInstanceEditAction.CanExecuteAction.BindLambda([bCanEdit](const FToolMenuContext&)
+		{
+			return bCanEdit;
+		});
+
+		FText EntryLabel = bSingleEntry ? LOCTEXT("EditLevelInstances", "Override") : FText::FromString(LevelInstance->GetWorldAsset().GetAssetName());
+		if (bCanEdit)
+		{
+			EntryDesc = FText::Format(LOCTEXT("LevelInstanceName", "{0}:{1}"), FText::FromString(LevelInstanceActor->GetActorLabel()), FText::FromString(LevelInstance->GetWorldAssetPackage()));
 		}
+		Section.AddMenuEntry(NAME_None, EntryLabel, EntryDesc, FSlateIcon(), LevelInstanceEditAction);
+	}
 
+	static void CreateEditPropertyOverridesSubMenu(UToolMenu* Menu, TArray<ILevelInstanceInterface*> LevelInstanceHierarchy, AActor* ContextActor)
+	{
+		FToolMenuSection& Section = Menu->AddSection(NAME_None, LOCTEXT("LevelInstanceContextEditSection", "Context"));
+		for (ILevelInstanceInterface* LevelInstance : LevelInstanceHierarchy)
+		{
+			CreateEditPropertyOverridesMenuEntry(Section, LevelInstance, ContextActor, false);
+		}
+	}
+		
+	static void MoveSelectionToLevelInstance(ILevelInstanceInterface* DestinationLevelInstance, const TArray<AActor*>& ActorsToMove)
+	{
 		DestinationLevelInstance->MoveActorsTo(ActorsToMove);
 	}
 		
-	void CreateEditMenu(UToolMenu* Menu, AActor* ContextActor)
+	static void CreateEditMenu(UToolMenu* Menu, AActor* ContextActor)
 	{
 		if (ULevelInstanceSubsystem* LevelInstanceSubsystem = ContextActor->GetWorld()->GetSubsystem<ULevelInstanceSubsystem>())
 		{
 			TArray<ILevelInstanceInterface*> LevelInstanceHierarchy;
 			LevelInstanceSubsystem->ForEachLevelInstanceAncestorsAndSelf(ContextActor, [&LevelInstanceHierarchy](ILevelInstanceInterface* AncestorLevelInstance)
 			{
-				if (IsExperimentalSettingEnabled(AncestorLevelInstance))
-				{
-					LevelInstanceHierarchy.Add(AncestorLevelInstance);
-				}
+				LevelInstanceHierarchy.Add(AncestorLevelInstance);
 				return true;
 			});
 
@@ -160,55 +169,125 @@ namespace LevelInstanceMenuUtils
 		}
 	}
 
-	void CreateCommitDiscardMenu(UToolMenu* Menu, AActor* ContextActor)
+	static void CreateEditPropertyOverridesMenu(UToolMenu* Menu, AActor* ContextActor)
+	{
+		if (!ULevelInstanceSettings::Get()->IsPropertyOverrideEnabled())
+		{
+			return;
+		}
+
+		if (ULevelInstanceSubsystem* LevelInstanceSubsystem = ContextActor->GetWorld()->GetSubsystem<ULevelInstanceSubsystem>())
+		{
+			TArray<ILevelInstanceInterface*> LevelInstanceHierarchy;
+			LevelInstanceSubsystem->ForEachLevelInstanceAncestorsAndSelf(ContextActor, [&LevelInstanceHierarchy](ILevelInstanceInterface* AncestorLevelInstance)
+			{
+				LevelInstanceHierarchy.Add(AncestorLevelInstance);
+				return true;
+			});
+
+			// Don't create sub menu if only one Level Instance is available to edit
+			if (LevelInstanceHierarchy.Num() == 1)
+			{
+				FToolMenuSection& Section = CreateLevelSection(Menu);
+				CreateEditPropertyOverridesMenuEntry(Section, LevelInstanceHierarchy[0], ContextActor, true);
+			}
+			else if (LevelInstanceHierarchy.Num() > 1)
+			{
+				FToolMenuSection& Section = CreateLevelSection(Menu);
+				Section.AddSubMenu(
+					"PropertyOverrideLevelInstances",
+					LOCTEXT("EditLevelInstancesPropertyOverrides", "Override"),
+					TAttribute<FText>(),
+					FNewToolMenuDelegate::CreateStatic(&CreateEditPropertyOverridesSubMenu, MoveTemp(LevelInstanceHierarchy), ContextActor)
+				);
+			}
+		}
+	}
+
+	static void CreateCommitDiscardMenu(UToolMenu* Menu, AActor* ContextActor)
 	{
 		ILevelInstanceInterface* LevelInstanceEdit = nullptr;
 		if (ContextActor)
 		{
 			if (ULevelInstanceSubsystem* LevelInstanceSubsystem = ContextActor->GetWorld()->GetSubsystem<ULevelInstanceSubsystem>())
 			{
-				LevelInstanceEdit = LevelInstanceSubsystem->GetEditingLevelInstance();
+				// Commit Property Overrides has priority
+				LevelInstanceEdit = LevelInstanceSubsystem->GetEditingPropertyOverridesLevelInstance();
+				if (!LevelInstanceEdit)
+				{
+					LevelInstanceEdit = LevelInstanceSubsystem->GetEditingLevelInstance();
+				}
 			}
 		}
 
+		// Commmit Property Overrides has priority
 		if (!LevelInstanceEdit)
 		{
 			if (ULevelInstanceSubsystem* LevelInstanceSubsystem = GEditor->GetEditorWorldContext().World()->GetSubsystem<ULevelInstanceSubsystem>())
 			{
-				LevelInstanceEdit = LevelInstanceSubsystem->GetEditingLevelInstance();
+				LevelInstanceEdit = LevelInstanceSubsystem->GetEditingPropertyOverridesLevelInstance();
 			}
 		}
 
+		// If no Property Overrides found try to find a regular Edit
+		if (!LevelInstanceEdit)
+		{
+			if (ULevelInstanceSubsystem* LevelInstanceSubsystem = GEditor->GetEditorWorldContext().World()->GetSubsystem<ULevelInstanceSubsystem>())
+			{
+				LevelInstanceEdit = LevelInstanceSubsystem->GetEditingLevelInstance();
+			}
+		}
+						
 		if (LevelInstanceEdit)
 		{
 			FToolMenuSection& Section = CreateLevelSection(Menu);
+			if (LevelInstanceEdit->IsEditingPropertyOverrides())
+			{
+				FText CommitTooltip;
+				const bool bCanCommit = LevelInstanceEdit->CanExitEditPropertyOverrides(/*bDiscardEdits=*/false, &CommitTooltip);
+
+				FToolUIAction CommitAction;
+				CommitAction.ExecuteAction.BindLambda([LevelInstanceEdit](const FToolMenuContext&) { LevelInstanceEdit->ExitEditPropertyOverrides(/*bDiscardEdits=*/false); });
+				CommitAction.CanExecuteAction.BindLambda([bCanCommit](const FToolMenuContext&) { return bCanCommit; });
+				Section.AddMenuEntry(NAME_None, LOCTEXT("LevelInstanceCommitPropertyOverridesLabel", "Commit Override(s)"), CommitTooltip, FSlateIcon(), CommitAction);
 
-			FText CommitTooltip;
-			const bool bCanCommit = LevelInstanceEdit->CanExitEdit(/*bDiscardEdits=*/false, &CommitTooltip);
+				FText DiscardTooltip;
+				const bool bCanDiscard = LevelInstanceEdit->CanExitEditPropertyOverrides(/*bDiscardEdits=*/true, &DiscardTooltip);
+
+				FToolUIAction DiscardAction;
+				DiscardAction.ExecuteAction.BindLambda([LevelInstanceEdit](const FToolMenuContext&) { LevelInstanceEdit->ExitEditPropertyOverrides(/*bDiscardEdits=*/true); });
+				DiscardAction.CanExecuteAction.BindLambda([bCanDiscard](const FToolMenuContext&) { return bCanDiscard; });
+				Section.AddMenuEntry(NAME_None, LOCTEXT("LevelInstanceDiscardPropertyOverridesLabel", "Discard Override(s)"), DiscardTooltip, FSlateIcon(), DiscardAction);
+			}
+			else
+			{
+				FText CommitTooltip;
+				const bool bCanCommit = LevelInstanceEdit->CanExitEdit(/*bDiscardEdits=*/false, &CommitTooltip);
 
-			FToolUIAction CommitAction;
-			CommitAction.ExecuteAction.BindLambda([LevelInstanceEdit](const FToolMenuContext&) { LevelInstanceEdit->ExitEdit(/*bDiscardEdits=*/false); });
-			CommitAction.CanExecuteAction.BindLambda([bCanCommit](const FToolMenuContext&) { return bCanCommit; });
-			Section.AddMenuEntry(NAME_None, LOCTEXT("LevelInstanceCommitLabel", "Commit"), CommitTooltip, FSlateIcon(), CommitAction);
+				FToolUIAction CommitAction;
+				CommitAction.ExecuteAction.BindLambda([LevelInstanceEdit](const FToolMenuContext&) { LevelInstanceEdit->ExitEdit(/*bDiscardEdits=*/false); });
+				CommitAction.CanExecuteAction.BindLambda([bCanCommit](const FToolMenuContext&) { return bCanCommit; });
+				Section.AddMenuEntry(NAME_None, LOCTEXT("LevelInstanceCommitLabel", "Commit"), CommitTooltip, FSlateIcon(), CommitAction);
 
-			FText DiscardTooltip;
-			const bool bCanDiscard = LevelInstanceEdit->CanExitEdit(/*bDiscardEdits=*/true, &DiscardTooltip);
+				FText DiscardTooltip;
+				const bool bCanDiscard = LevelInstanceEdit->CanExitEdit(/*bDiscardEdits=*/true, &DiscardTooltip);
 
-			FToolUIAction DiscardAction;
-			DiscardAction.ExecuteAction.BindLambda([LevelInstanceEdit](const FToolMenuContext&) { LevelInstanceEdit->ExitEdit(/*bDiscardEdits=*/true); });
-			DiscardAction.CanExecuteAction.BindLambda([bCanDiscard](const FToolMenuContext&) { return bCanDiscard; });
-			Section.AddMenuEntry(NAME_None, LOCTEXT("LevelInstanceDiscardLabel", "Discard"), DiscardTooltip, FSlateIcon(), DiscardAction);
+				FToolUIAction DiscardAction;
+				DiscardAction.ExecuteAction.BindLambda([LevelInstanceEdit](const FToolMenuContext&) { LevelInstanceEdit->ExitEdit(/*bDiscardEdits=*/true); });
+				DiscardAction.CanExecuteAction.BindLambda([bCanDiscard](const FToolMenuContext&) { return bCanDiscard; });
+				Section.AddMenuEntry(NAME_None, LOCTEXT("LevelInstanceDiscardLabel", "Discard"), DiscardTooltip, FSlateIcon(), DiscardAction);
+			}
 		}
 	}
 
-	void CreateSetCurrentMenu(UToolMenu* Menu, AActor* ContextActor)
+	static void CreateSetCurrentMenu(UToolMenu* Menu, AActor* ContextActor)
 	{
 		ILevelInstanceInterface* LevelInstanceEdit = nullptr;
 		if (ULevelInstanceSubsystem* LevelInstanceSubsystem = ContextActor->GetWorld()->GetSubsystem<ULevelInstanceSubsystem>())
 		{
 			LevelInstanceEdit = LevelInstanceSubsystem->GetEditingLevelInstance();
 
-			if (LevelInstanceEdit)
+			if (LevelInstanceEdit && !LevelInstanceSubsystem->GetEditingPropertyOverridesLevelInstance())
 			{
 				FToolUIAction LevelInstanceSetCurrentAction;
 				LevelInstanceSetCurrentAction.ExecuteAction.BindLambda([LevelInstanceEdit](const FToolMenuContext&)
@@ -222,9 +301,9 @@ namespace LevelInstanceMenuUtils
 		}
 	}
 
-	void CreateMoveSelectionToMenu(UToolMenu* Menu)
+	static void CreateMoveSelectionToMenu(UToolMenu* Menu, const TArray<AActor*>& ActorsToMove)
 	{
-		if (GEditor->GetSelectedActorCount() > 0)
+		if (ActorsToMove.Num() > 0)
 		{
 			ILevelInstanceInterface* LevelInstanceEdit = nullptr;
 			ULevelInstanceSubsystem* LevelInstanceSubsystem = GEditor->GetEditorWorldContext().World()->GetSubsystem<ULevelInstanceSubsystem>();
@@ -233,30 +312,27 @@ namespace LevelInstanceMenuUtils
 				LevelInstanceEdit = LevelInstanceSubsystem->GetEditingLevelInstance();
 			}
 			
-			if (LevelInstanceEdit)
+			if (LevelInstanceEdit && !LevelInstanceSubsystem->GetEditingPropertyOverridesLevelInstance())
 			{
 				FToolUIAction LevelInstanceMoveSelectionAction;
 
-				LevelInstanceMoveSelectionAction.CanExecuteAction.BindLambda([LevelInstanceEdit, LevelInstanceSubsystem](const FToolMenuContext& MenuContext)
+				LevelInstanceMoveSelectionAction.CanExecuteAction.BindLambda([CopyActorsToMove = ActorsToMove, LevelInstanceEdit, LevelInstanceSubsystem](const FToolMenuContext& MenuContext)
+				{
+					for (AActor* ActorToMove : CopyActorsToMove)
 					{
-						for (FSelectionIterator It(GEditor->GetSelectedActorIterator()); It; ++It)
+						if (ActorToMove->GetLevel() == LevelInstanceSubsystem->GetLevelInstanceLevel(LevelInstanceEdit))
 						{
-							if (AActor* Actor = Cast<AActor>(*It))
-							{
-								if (Actor->GetLevel() == LevelInstanceSubsystem->GetLevelInstanceLevel(LevelInstanceEdit))
-								{
-									return false;
-								}
-							}
+							return false;
 						}
+					}
 
-						return GEditor->GetSelectedActorCount() > 0;
-					});
+					return !CopyActorsToMove.IsEmpty();
+				});
 
-				LevelInstanceMoveSelectionAction.ExecuteAction.BindLambda([LevelInstanceEdit](const FToolMenuContext&)
-					{
-						MoveSelectionToLevelInstance(LevelInstanceEdit);
-					});
+				LevelInstanceMoveSelectionAction.ExecuteAction.BindLambda([CopyActorsToMove = ActorsToMove, LevelInstanceEdit](const FToolMenuContext&)
+				{
+					MoveSelectionToLevelInstance(LevelInstanceEdit, CopyActorsToMove);
+				});
 
 				FToolMenuSection& Section = CreateLevelSection(Menu);
 				Section.AddMenuEntry(TEXT("LevelInstanceMoveSelectionTo"), LOCTEXT("LevelInstanceMoveSelectionTo", "Move Selection to"), TAttribute<FText>(), FSlateIcon(), LevelInstanceMoveSelectionAction);
@@ -265,7 +341,7 @@ namespace LevelInstanceMenuUtils
 		}
 	}
 
-	UClass* GetDefaultLevelInstanceClass(ELevelInstanceCreationType CreationType)
+	static UClass* GetDefaultLevelInstanceClass(ELevelInstanceCreationType CreationType)
 	{
 		if (CreationType == ELevelInstanceCreationType::PackedLevelActor)
 		{
@@ -285,34 +361,21 @@ namespace LevelInstanceMenuUtils
 		return ALevelInstance::StaticClass();
 	}
 
-	bool AreAllSelectedLevelInstancesRootSelections()
+	static bool AreAllSelectedLevelInstancesRootSelections(const TArray<ILevelInstanceInterface*>& SelectedLevelInstances)
 	{
-		for (FSelectionIterator It(GEditor->GetSelectedActorIterator()); It; ++It)
+		for(ILevelInstanceInterface* LevelInstance : SelectedLevelInstances)
 		{
-			if (ILevelInstanceInterface* LevelInstance = Cast<ILevelInstanceInterface>(*It))
+			if (CastChecked<AActor>(LevelInstance)->GetSelectionParent() != nullptr)
 			{
-				if (CastChecked<AActor>(*It)->GetSelectionParent() != nullptr)
-				{
-					return false;
-				}
+				return false;
 			}
 		}
 
 		return true;
 	}
 		
-	void CreateLevelInstanceFromSelection(ULevelInstanceSubsystem* LevelInstanceSubsystem, ELevelInstanceCreationType CreationType)
+	static void CreateLevelInstanceFromSelection(ULevelInstanceSubsystem* LevelInstanceSubsystem, ELevelInstanceCreationType CreationType, const TArray<AActor*>& ActorsToMove)
 	{
-		TArray<AActor*> ActorsToMove;
-		ActorsToMove.Reserve(GEditor->GetSelectedActorCount());
-		for (FSelectionIterator It(GEditor->GetSelectedActorIterator()); It; ++It)
-		{
-			if (AActor* Actor = Cast<AActor>(*It))
-			{
-				ActorsToMove.Add(Actor);
-			}
-		}
-
 		IMainFrameModule& MainFrameModule = FModuleManager::GetModuleChecked<IMainFrameModule>("MainFrame");
 
 		TSharedPtr<SWindow> NewLevelInstanceWindow =
@@ -368,38 +431,39 @@ namespace LevelInstanceMenuUtils
 		}
 	}
 
-	void CreateCreateMenu(UToolMenu* ToolMenu)
+	static void CreateCreateMenu(UToolMenu* ToolMenu, const TArray<AActor*>& ActorsToMove)
 	{
 		if (ULevelInstanceSubsystem* LevelInstanceSubsystem = GEditor->GetEditorWorldContext().World()->GetSubsystem<ULevelInstanceSubsystem>())
 		{
-			if (GEditor->GetSelectedActorCount() > 0)
+			if (LevelInstanceSubsystem->CanCreateLevelInstanceFrom(ActorsToMove))
 			{
 				FToolMenuSection& Section = ToolMenu->AddSection("ActorSelectionSectionName", LOCTEXT("ActorSelectionSectionLabel", "Actor Selection"));
-
-				if (GetDefault<UEditorExperimentalSettings>()->bLevelInstance)
-				{
-					Section.AddMenuEntry(
-						TEXT("CreateLevelInstance"),
-						FText::Format(LOCTEXT("CreateFromSelectionLabel", "Create {0}..."), StaticEnum<ELevelInstanceCreationType>()->GetDisplayNameTextByValue((int64)ELevelInstanceCreationType::LevelInstance)),
-						TAttribute<FText>(),
-						FSlateIcon(FAppStyle::GetAppStyleSetName(), "ClassIcon.LevelInstance"),
-						FExecuteAction::CreateStatic(&LevelInstanceMenuUtils::CreateLevelInstanceFromSelection, LevelInstanceSubsystem, ELevelInstanceCreationType::LevelInstance));
-				}
-
-				if (GetDefault<UEditorExperimentalSettings>()->bPackedLevelActor)
-				{
-					Section.AddMenuEntry(
-						TEXT("CreatePackedLevelBlueprint"),
-						FText::Format(LOCTEXT("CreateFromSelectionLabel", "Create {0}..."), StaticEnum<ELevelInstanceCreationType>()->GetDisplayNameTextByValue((int64)ELevelInstanceCreationType::PackedLevelActor)),
-						TAttribute<FText>(),
-						FSlateIcon(FAppStyle::GetAppStyleSetName(), "ClassIcon.PackedLevelActor"),
-						FExecuteAction::CreateStatic(&LevelInstanceMenuUtils::CreateLevelInstanceFromSelection, LevelInstanceSubsystem, ELevelInstanceCreationType::PackedLevelActor));
-				}
+								
+				Section.AddMenuEntry(
+					TEXT("CreateLevelInstance"),
+					FText::Format(LOCTEXT("CreateFromSelectionLabel", "Create {0}..."), StaticEnum<ELevelInstanceCreationType>()->GetDisplayNameTextByValue((int64)ELevelInstanceCreationType::LevelInstance)),
+					TAttribute<FText>(),
+					FSlateIcon(FAppStyle::GetAppStyleSetName(), "ClassIcon.LevelInstance"),
+					FExecuteAction::CreateLambda([LevelInstanceSubsystem, CopyActorsToMove = ActorsToMove]
+					{
+							CreateLevelInstanceFromSelection(LevelInstanceSubsystem, ELevelInstanceCreationType::LevelInstance, CopyActorsToMove);
+					}));
+				
+				Section.AddMenuEntry(
+					TEXT("CreatePackedLevelBlueprint"),
+					FText::Format(LOCTEXT("CreateFromSelectionLabel", "Create {0}..."), StaticEnum<ELevelInstanceCreationType>()->GetDisplayNameTextByValue((int64)ELevelInstanceCreationType::PackedLevelActor)),
+					TAttribute<FText>(),
+					FSlateIcon(FAppStyle::GetAppStyleSetName(), "ClassIcon.PackedLevelActor"),
+					FExecuteAction::CreateLambda([LevelInstanceSubsystem, CopyActorsToMove = ActorsToMove]
+					{
+						CreateLevelInstanceFromSelection(LevelInstanceSubsystem, ELevelInstanceCreationType::PackedLevelActor, CopyActorsToMove);
+					}));
+				
 			}
 		}
 	}
 		
-	void CreateBreakSubMenu(UToolMenu* Menu, TArray<ILevelInstanceInterface*> BreakableLevelInstances)
+	static void CreateBreakSubMenu(UToolMenu* Menu, const TArray<ILevelInstanceInterface*>& BreakableLevelInstances)
 	{
 		static int32 BreakLevels = 1;
 		ULevelInstanceEditorPerProjectUserSettings* Settings = GetMutableDefault<ULevelInstanceEditorPerProjectUserSettings>();
@@ -457,14 +521,14 @@ namespace LevelInstanceMenuUtils
 				FText(),
 				FSlateIcon(),
 				FUIAction(
-					FExecuteAction::CreateLambda([BreakableLevelInstances, LevelInstanceSubsystem, Settings]()
+					FExecuteAction::CreateLambda([CopyBreakableLevelInstances = BreakableLevelInstances, LevelInstanceSubsystem, Settings]()
 					{
 						const FText LevelInstanceBreakWarning = FText::Format(
 							LOCTEXT(
 								"BreakingLevelInstance",
 								"You are about to break {0} level instance(s). This action cannot be undone. Are you sure ?"
 							),
-							FText::AsNumber(BreakableLevelInstances.Num())
+							FText::AsNumber(CopyBreakableLevelInstances.Num())
 						);
 
 						if (FMessageDialog::Open(EAppMsgType::YesNo, LevelInstanceBreakWarning) == EAppReturnType::Yes)
@@ -475,7 +539,7 @@ namespace LevelInstanceMenuUtils
 								Flags |= ELevelInstanceBreakFlags::KeepFolders;
 							}
 
-							for (ILevelInstanceInterface* LevelInstance : BreakableLevelInstances)
+							for (ILevelInstanceInterface* LevelInstance : CopyBreakableLevelInstances)
 							{
 								LevelInstanceSubsystem->BreakLevelInstance(LevelInstance, BreakLevels, nullptr, Flags);
 							}
@@ -489,25 +553,19 @@ namespace LevelInstanceMenuUtils
 		}
 	}
 
-	void CreateBreakMenu(UToolMenu* Menu)
+	static void CreateBreakMenu(UToolMenu* Menu, const TArray<ILevelInstanceInterface*>& SelectedLevelInstances)
 	{
 		if(ULevelInstanceSubsystem* LevelInstanceSubsystem = GEditor->GetEditorWorldContext().World()->GetSubsystem<ULevelInstanceSubsystem>())
 		{
 			TArray<ILevelInstanceInterface*> BreakableLevelInstances;
-			if (GEditor->GetSelectedActorCount() > 0)
+			for (ILevelInstanceInterface* SelectedLevelInstance : SelectedLevelInstances)
 			{
-				for (FSelectionIterator It(GEditor->GetSelectedActorIterator()); It; ++It)
+				if (LevelInstanceSubsystem->CanBreakLevelInstance(SelectedLevelInstance))
 				{
-					if (ILevelInstanceInterface* LevelInstance = Cast<ILevelInstanceInterface>(*It))
-					{
-						if (IsExperimentalSettingEnabled(LevelInstance) && !LevelInstanceSubsystem->IsEditingLevelInstance(LevelInstance) && !LevelInstanceSubsystem->LevelInstanceHasLevelScriptBlueprint(LevelInstance))
-						{
-							BreakableLevelInstances.Add(LevelInstance);
-						}	
-					}
-				}
+					BreakableLevelInstances.Add(SelectedLevelInstance);
+				}	
 			}
-
+			
 			if (BreakableLevelInstances.Num())
 			{
 				FToolMenuSection& Section = CreateLevelSection(Menu);
@@ -515,14 +573,16 @@ namespace LevelInstanceMenuUtils
 					"BreakLevelInstances",
 					LOCTEXT("BreakLevelInstances", "Break..."),
 					TAttribute<FText>(),
-					FNewToolMenuDelegate::CreateStatic(&CreateBreakSubMenu, BreakableLevelInstances)
-				);
+					FNewToolMenuDelegate::CreateLambda([CopyOfBreakableLevelInstances = BreakableLevelInstances](UToolMenu* Menu)
+					{
+						CreateBreakSubMenu(Menu, CopyOfBreakableLevelInstances);
+					}));
 			}
 		}
 
 	}
 
-	void CreatePackedBlueprintMenu(UToolMenu* Menu, AActor* ContextActor)
+	static void CreatePackedBlueprintMenu(UToolMenu* Menu, AActor* ContextActor)
 	{
 		if (ULevelInstanceSubsystem* LevelInstanceSubsystem = ContextActor->GetWorld()->GetSubsystem<ULevelInstanceSubsystem>())
 		{
@@ -539,7 +599,7 @@ namespace LevelInstanceMenuUtils
 				return true;
 			});
 						
-			if (ContextLevelInstance && IsExperimentalSettingEnabled(ContextLevelInstance) && !ContextLevelInstance->IsEditing())
+			if (ContextLevelInstance && !ContextLevelInstance->IsEditing())
 			{
 				FToolMenuSection& Section = CreateLevelSection(Menu);
 				;
@@ -569,22 +629,98 @@ namespace LevelInstanceMenuUtils
 		}
 	}
 
+	static void CreateResetPropertyOverridesMenu(UToolMenu* Menu, const TArray<AActor*>& SelectedActors, const TArray<ILevelInstanceInterface*>& SelectedLevelInstances)
+	{
+		if (!ULevelInstanceSettings::Get()->IsPropertyOverrideEnabled())
+		{
+			return;
+		}
+
+		if (ULevelInstanceSubsystem* LevelInstanceSubsystem = UWorld::GetSubsystem<ULevelInstanceSubsystem>(GEditor->GetEditorWorldContext().World()))
+		{
+			if (SelectedLevelInstances.Num() > 0 && SelectedActors.Num() == SelectedLevelInstances.Num())
+			{
+				bool bCanResetAllLevelInstances = true;
+				for (ILevelInstanceInterface* SelectedLevelInstance : SelectedLevelInstances)
+				{
+					if (!LevelInstanceSubsystem->CanResetPropertyOverrides(SelectedLevelInstance))
+					{
+						return;
+					}
+				}
+
+				FToolMenuSection& Section = CreateLevelSection(Menu);
+				FToolUIAction UIAction;
+				UIAction.ExecuteAction.BindLambda([LevelInstanceSubsystem, CopySelectedLevelInstance = SelectedLevelInstances](const FToolMenuContext& MenuContext)
+				{
+					for (ILevelInstanceInterface* LevelInstanceInterface : CopySelectedLevelInstance)
+					{
+						LevelInstanceSubsystem->ResetPropertyOverrides(LevelInstanceInterface);
+					}
+				});
+
+				Section.AddMenuEntry(
+					"ResetLevelInstancePropertyOverrides",
+					LOCTEXT("ResetLevelInstancePropertyOverrides", "Reset Override(s)"),
+					TAttribute<FText>(),
+					TAttribute<FSlateIcon>(),
+					UIAction);
+
+				return;
+			}
+
+			if (SelectedActors.Num() > 0)
+			{
+				bool bCanResetAllActors = true;
+				for (AActor* SelectedActor : SelectedActors)
+				{
+					if (!LevelInstanceSubsystem->CanResetPropertyOverridesForActor(SelectedActor))
+					{
+						bCanResetAllActors = false;
+						break;
+					}
+				}
+
+				if (bCanResetAllActors)
+				{
+					FToolMenuSection& Section = CreateLevelSection(Menu);
+					FToolUIAction UIAction;
+					UIAction.ExecuteAction.BindLambda([LevelInstanceSubsystem, CopySelectedActors = SelectedActors](const FToolMenuContext& MenuContext)
+					{
+						FScopedTransaction ResetPropertyOverridesTransaction(LOCTEXT("ResetPropertyOverrides", "Reset Property Override(s)"));
+						for (AActor* SelectedActor : CopySelectedActors)
+						{
+							LevelInstanceSubsystem->ResetPropertyOverridesForActor(SelectedActor);
+						}
+					});
+
+					Section.AddMenuEntry(
+						"ResetLevelInstancePropertyOverrides",
+						LOCTEXT("ResetLevelInstancePropertyOverrides", "Reset Override(s)"),
+						TAttribute<FText>(),
+						TAttribute<FSlateIcon>(),
+						UIAction);
+				}
+			}
+		}
+	}
+
 	class FLevelInstanceClassFilter : public IClassViewerFilter
 	{
 	public:
 		
-		virtual bool IsClassAllowed(const FClassViewerInitializationOptions& InInitOptions, const UClass* InClass, TSharedRef< FClassViewerFilterFuncs > InFilterFuncs) override
+		virtual bool IsClassAllowed(const FClassViewerInitializationOptions& InInitOptions, const UClass* InClass, TSharedRef<FClassViewerFilterFuncs> InFilterFuncs) override
 		{
 			return InClass && InClass->ImplementsInterface(ULevelInstanceInterface::StaticClass()) && InClass->IsNative() && !InClass->HasAnyClassFlags(CLASS_Deprecated);
 		}
 
-		virtual bool IsUnloadedClassAllowed(const FClassViewerInitializationOptions& InInitOptions, const TSharedRef< const IUnloadedBlueprintData > InUnloadedClassData, TSharedRef< FClassViewerFilterFuncs > InFilterFuncs) override
+		virtual bool IsUnloadedClassAllowed(const FClassViewerInitializationOptions& InInitOptions, const TSharedRef<const IUnloadedBlueprintData> InUnloadedClassData, TSharedRef<FClassViewerFilterFuncs> InFilterFuncs) override
 		{
 			return false;
 		}
 	};
 
-	void CreateBlueprintFromWorld(UWorld* WorldAsset)
+	static void CreateBlueprintFromWorld(UWorld* WorldAsset)
 	{
 		TSoftObjectPtr<UWorld> LevelInstancePtr(WorldAsset);
 
@@ -631,7 +767,7 @@ namespace LevelInstanceMenuUtils
 		}		
 	}
 
-	void CreateBlueprintFromMenu(UToolMenu* Menu, FAssetData WorldAsset)
+	static void CreateBlueprintFromMenu(UToolMenu* Menu, FAssetData WorldAsset)
 	{
 		FToolMenuSection& Section = CreateLevelSection(Menu);
 		FToolUIAction UIAction;
@@ -651,7 +787,7 @@ namespace LevelInstanceMenuUtils
 			UIAction);
 	}
 
-	void AddPartitionedStreamingSupportFromWorld(UWorld* WorldAsset)
+	static void AddPartitionedStreamingSupportFromWorld(UWorld* WorldAsset)
 	{
 		if (WorldAsset->GetStreamingLevels().Num())
 		{
@@ -709,7 +845,7 @@ namespace LevelInstanceMenuUtils
 		}
 	}
 
-	void UpdatePackedBlueprintsFromMenu(UToolMenu* Menu, FAssetData WorldAsset)
+	static void UpdatePackedBlueprintsFromMenu(UToolMenu* Menu, FAssetData WorldAsset)
 	{
 		FToolMenuSection& Section = CreateLevelSection(Menu);
 		FToolUIAction UIAction;
@@ -741,7 +877,7 @@ namespace LevelInstanceMenuUtils
 		);
 	}
 
-	void AddPartitionedStreamingSupportFromMenu(UToolMenu* Menu, FAssetData WorldAsset)
+	static void AddPartitionedStreamingSupportFromMenu(UToolMenu* Menu, FAssetData WorldAsset)
 	{
 		FName WorldAssetName = WorldAsset.PackageName;
 		if (!ULevel::GetIsLevelPartitionedFromPackage(WorldAssetName))
@@ -855,16 +991,13 @@ void FLevelInstanceEditorModule::BroadcastTryExitEditorMode()
 	TryExitEditorModeEvent.Broadcast();
 }
 
-void FLevelInstanceEditorModule::ActivateEditorMode()
+void FLevelInstanceEditorModule::UpdateEditorMode(bool bActivated)
 {
-	if (!GLevelEditorModeTools().IsModeActive(ULevelInstanceEditorMode::EM_LevelInstanceEditorModeId))
+	if (bActivated && !GLevelEditorModeTools().IsModeActive(ULevelInstanceEditorMode::EM_LevelInstanceEditorModeId))
 	{
 		GLevelEditorModeTools().ActivateMode(ULevelInstanceEditorMode::EM_LevelInstanceEditorModeId);
 	}
-}
-void FLevelInstanceEditorModule::DeactivateEditorMode()
-{
-	if (GLevelEditorModeTools().IsModeActive(ULevelInstanceEditorMode::EM_LevelInstanceEditorModeId))
+	else if(!bActivated && GLevelEditorModeTools().IsModeActive(ULevelInstanceEditorMode::EM_LevelInstanceEditorModeId))
 	{
 		GLevelEditorModeTools().DeactivateMode(ULevelInstanceEditorMode::EM_LevelInstanceEditorModeId);
 	}
@@ -908,10 +1041,7 @@ void FLevelInstanceEditorModule::ExtendContextMenu()
 				return FPackedLevelActorUtils::CanPack();
 			}),
 			FIsActionChecked(),
-			FIsActionButtonVisible::CreateLambda([]() 
-			{
-				return GetDefault<UEditorExperimentalSettings>()->bPackedLevelActor;
-			}));
+			FIsActionButtonVisible());
 
 		FToolMenuEntry& Entry = Section.AddMenuEntry(NAME_None, LOCTEXT("PackLevelActorsTitle", "Pack Level Actors"),
 			LOCTEXT("PackLevelActorsTooltip", "Update packed level actor blueprints"), FSlateIcon(), PackAction, EUserInterfaceActionType::Button);
@@ -924,42 +1054,67 @@ void FLevelInstanceEditorModule::ExtendContextMenu()
 			return;
 		}
 
+		if (GetDefault<ULevelInstanceSettings>()->IsLevelInstanceDisabled())
+		{
+			return;
+		}
+
+		// Build Selection for Menus
+		TArray<AActor*> SelectedActors;
+		TArray<ILevelInstanceInterface*> SelectedLevelInstances;
+		SelectedActors.Reserve(GEditor->GetSelectedActorCount());
+		SelectedLevelInstances.Reserve(GEditor->GetSelectedActorCount());
+		for (FSelectionIterator It(GEditor->GetSelectedActorIterator()); It; ++It)
+		{
+			if (AActor* Actor = Cast<AActor>(*It); IsValid(Actor))
+			{
+				SelectedActors.Add(Actor);
+
+				if (ILevelInstanceInterface* LevelInstance = Cast<ILevelInstanceInterface>(Actor))
+				{
+					SelectedLevelInstances.Add(LevelInstance);
+				}
+			}
+		}
+
 		// Some actions aren't allowed on non root selection Level Instances (Readonly Level Instances)
-		const bool bAreAllSelectedLevelInstancesRootSelections = LevelInstanceMenuUtils::AreAllSelectedLevelInstancesRootSelections();
+		const bool bAreAllSelectedLevelInstancesRootSelections = FLevelInstanceMenuUtils::AreAllSelectedLevelInstancesRootSelections(SelectedLevelInstances);
 
 		if (ULevelEditorContextMenuContext* LevelEditorMenuContext = ToolMenu->Context.FindContext<ULevelEditorContextMenuContext>())
 		{
 			// Use the actor under the cursor if available (e.g. right-click menu).
 			// Otherwise use the first selected actor if there's one (e.g. Actor pulldown menu or outliner).
 			AActor* ContextActor = LevelEditorMenuContext->HitProxyActor.Get();
-			if (!ContextActor && GEditor->GetSelectedActorCount() != 0)
+			if (!ContextActor && SelectedActors.Num() > 0)
 			{
-				ContextActor = Cast<AActor>(GEditor->GetSelectedActors()->GetSelectedObject(0));
+				ContextActor = SelectedActors[0];
 			}
 
 			if (ContextActor)
 			{
 				// Allow Edit/Commmit on non root selected Level Instance
-				LevelInstanceMenuUtils::CreateEditMenu(ToolMenu, ContextActor);
-				LevelInstanceMenuUtils::CreateCommitDiscardMenu(ToolMenu, ContextActor);
+				FLevelInstanceMenuUtils::CreateEditMenu(ToolMenu, ContextActor);
+				FLevelInstanceMenuUtils::CreateEditPropertyOverridesMenu(ToolMenu, ContextActor);
+				FLevelInstanceMenuUtils::CreateCommitDiscardMenu(ToolMenu, ContextActor);
 				
 				if (bAreAllSelectedLevelInstancesRootSelections)
 				{
-					LevelInstanceMenuUtils::CreatePackedBlueprintMenu(ToolMenu, ContextActor);
-					LevelInstanceMenuUtils::CreateSetCurrentMenu(ToolMenu, ContextActor);
+					FLevelInstanceMenuUtils::CreatePackedBlueprintMenu(ToolMenu, ContextActor);
+					FLevelInstanceMenuUtils::CreateSetCurrentMenu(ToolMenu, ContextActor);
 				}
 			}
 
 			if (bAreAllSelectedLevelInstancesRootSelections)
 			{
-				LevelInstanceMenuUtils::CreateMoveSelectionToMenu(ToolMenu);
+				FLevelInstanceMenuUtils::CreateMoveSelectionToMenu(ToolMenu, SelectedActors);
 			}
 		}
 
 		if (bAreAllSelectedLevelInstancesRootSelections)
 		{
-			LevelInstanceMenuUtils::CreateBreakMenu(ToolMenu);
-			LevelInstanceMenuUtils::CreateCreateMenu(ToolMenu);
+			FLevelInstanceMenuUtils::CreateBreakMenu(ToolMenu, SelectedLevelInstances);
+			FLevelInstanceMenuUtils::CreateCreateMenu(ToolMenu, SelectedActors);
+			FLevelInstanceMenuUtils::CreateResetPropertyOverridesMenu(ToolMenu, SelectedActors, SelectedLevelInstances);
 		}
 	};
 
@@ -982,7 +1137,7 @@ void FLevelInstanceEditorModule::ExtendContextMenu()
 				return;
 			}
 
-			if(!GetDefault<UEditorExperimentalSettings>()->bLevelInstance)
+			if (GetDefault<ULevelInstanceSettings>()->IsLevelInstanceDisabled())
 			{
 				return;
 			}
@@ -999,9 +1154,9 @@ void FLevelInstanceEditorModule::ExtendContextMenu()
 					const FAssetData& WorldAsset = AssetMenuContext->SelectedAssets[0];
 					if (AssetMenuContext->SelectedAssets[0].IsInstanceOf<UWorld>())
 					{
-						LevelInstanceMenuUtils::CreateBlueprintFromMenu(ToolMenu, WorldAsset);
-						LevelInstanceMenuUtils::UpdatePackedBlueprintsFromMenu(ToolMenu, WorldAsset);
-						LevelInstanceMenuUtils::AddPartitionedStreamingSupportFromMenu(ToolMenu, WorldAsset);
+						FLevelInstanceMenuUtils::CreateBlueprintFromMenu(ToolMenu, WorldAsset);
+						FLevelInstanceMenuUtils::UpdatePackedBlueprintsFromMenu(ToolMenu, WorldAsset);
+						FLevelInstanceMenuUtils::AddPartitionedStreamingSupportFromMenu(ToolMenu, WorldAsset);
 					}
 				}
 			}
diff --git a/Engine/Source/Editor/LevelInstanceEditor/Private/LevelInstanceEditorModule.h b/Engine/Source/Editor/LevelInstanceEditor/Private/LevelInstanceEditorModule.h
--- a/Engine/Source/Editor/LevelInstanceEditor/Private/LevelInstanceEditorModule.h
+++ b/Engine/Source/Editor/LevelInstanceEditor/Private/LevelInstanceEditorModule.h
@@ -3,6 +3,8 @@
 
 #include "CoreMinimal.h"
 #include "LevelInstance/ILevelInstanceEditorModule.h"
+#include "PropertyEditorArchetypePolicy.h"
+#include "PropertyEditorEditConstPolicy.h"
 #include "Tools/Modes.h"
 
 class AActor;
@@ -27,9 +29,6 @@ public:
 	 */
 	virtual void ShutdownModule();
 
-	virtual void ActivateEditorMode() override;
-	virtual void DeactivateEditorMode() override;
-		
 	virtual void BroadcastTryExitEditorMode() override;
 
 	DECLARE_DERIVED_EVENT(FLevelInstanceEditorModule, ILevelInstanceEditorModule::FExitEditorModeEvent, FExitEditorModeEvent);
@@ -41,12 +40,76 @@ public:
 	virtual bool IsEditInPlaceStreamingEnabled() const override;
 
 private:
+	virtual void UpdateEditorMode(bool bActivated) override;
+	
 	void OnEditorModeIDChanged(const FEditorModeID& InModeID, bool bIsEnteringMode);
 	void OnLevelActorDeleted(AActor* Actor);
 	void CanMoveActorToLevel(const AActor* ActorToMove, const ULevel* DestLevel, bool& bOutCanMove);
 
 	void ExtendContextMenu();
+		
+	class FPropertyEditorPolicy : public PropertyEditorPolicy::IEditConstPolicy, 
+								  public PropertyEditorPolicy::IArchetypePolicy
+	{
+	public:
+		FPropertyEditorPolicy(ILevelInstanceEditorModule::IPropertyOverridePolicy* InPropertyOverridePolicy)
+			: PropertyOverridePolicy(InPropertyOverridePolicy)
+		{
+			check(PropertyOverridePolicy);
+			PropertyEditorPolicy::RegisterEditConstPolicy(this);
+			PropertyEditorPolicy::RegisterArchetypePolicy(this);
+		}
+
+		virtual ~FPropertyEditorPolicy()
+		{
+			PropertyEditorPolicy::UnregisterEditConstPolicy(this);
+			PropertyEditorPolicy::UnregisterArchetypePolicy(this);
+		}
+
+		virtual UObject* GetArchetypeForObject(const UObject* Object) const override
+		{
+			return PropertyOverridePolicy->GetArchetypeForObject(Object);
+		}
+
+		virtual bool CanEditProperty(const FEditPropertyChain& PropertyChain, const UObject* Object) const override
+		{
+			return PropertyOverridePolicy->CanEditProperty(PropertyChain, Object);
+		}
+
+		virtual bool CanEditProperty(const FProperty* Property, const UObject* Object) const override
+		{
+			return PropertyOverridePolicy->CanEditProperty(Property, Object);
+		}
+
+		ILevelInstanceEditorModule::IPropertyOverridePolicy* PropertyOverridePolicy = nullptr;
+	};
+
+	TUniquePtr<FPropertyEditorPolicy> PropertyEditorPolicy;
+		
+	virtual bool IsPropertyEditConst(const FEditPropertyChain& PropertyChain, UObject* Object) override
+	{
+		return PropertyEditorPolicy::IsPropertyEditConst(PropertyChain, Object);
+	}
+	
+	virtual bool IsPropertyEditConst(const FProperty* Property, UObject* Object) override
+	{
+		return PropertyEditorPolicy::IsPropertyEditConst(Property, Object);
+	}
 
+	virtual UObject* GetArchetype(const UObject* Object) override
+	{
+		return PropertyEditorPolicy::GetArchetype(Object);
+	}
+	
+	virtual void SetPropertyOverridePolicy(ILevelInstanceEditorModule::IPropertyOverridePolicy* InPropertyOverridePolicy) override
+	{
+		PropertyEditorPolicy.Reset();
+		if (InPropertyOverridePolicy)
+		{
+			PropertyEditorPolicy = MakeUnique<FPropertyEditorPolicy>(InPropertyOverridePolicy);
+		}
+	}
+	
 	FExitEditorModeEvent ExitEditorModeEvent;
 	FTryExitEditorModeEvent TryExitEditorModeEvent;
 };
diff --git a/Engine/Source/Editor/SceneOutliner/Private/ActorBrowsingMode.cpp b/Engine/Source/Editor/SceneOutliner/Private/ActorBrowsingMode.cpp
--- a/Engine/Source/Editor/SceneOutliner/Private/ActorBrowsingMode.cpp
+++ b/Engine/Source/Editor/SceneOutliner/Private/ActorBrowsingMode.cpp
@@ -265,7 +265,7 @@ FActorBrowsingMode::FActorBrowsingMode(SSceneOutliner* InSceneOutliner, TWeakObj
 					// then the actor should not be selectable.
 					if (const ILevelInstanceInterface* ParentLevelInstance = LevelInstanceSubsystem->GetParentLevelInstance(Actor))
 					{
-						if (!LevelInstanceSubsystem->IsEditingLevelInstance(ParentLevelInstance))
+						if (!LevelInstanceSubsystem->IsEditingLevelInstance(ParentLevelInstance) && !LevelInstanceSubsystem->IsEditingLevelInstancePropertyOverrides(ParentLevelInstance))
 						{
 							return false;
 						}
diff --git a/Engine/Source/Editor/SceneOutliner/Private/ActorTreeItem.cpp b/Engine/Source/Editor/SceneOutliner/Private/ActorTreeItem.cpp
--- a/Engine/Source/Editor/SceneOutliner/Private/ActorTreeItem.cpp
+++ b/Engine/Source/Editor/SceneOutliner/Private/ActorTreeItem.cpp
@@ -231,6 +231,10 @@ private:
 			{
 				return FAppStyle::Get().GetSlateColor("Colors.AccentGreen");
 			}
+			else if (LevelInstance->IsEditingPropertyOverrides())
+			{
+				return FAppStyle::Get().GetSlateColor("Colors.AccentBlue");
+			}
 		}
 
 		auto TreeItem = TreeItemPtr.Pin();
diff --git a/Engine/Source/Editor/UnrealEd/Classes/Settings/EditorExperimentalSettings.h b/Engine/Source/Editor/UnrealEd/Classes/Settings/EditorExperimentalSettings.h
--- a/Engine/Source/Editor/UnrealEd/Classes/Settings/EditorExperimentalSettings.h
+++ b/Engine/Source/Editor/UnrealEd/Classes/Settings/EditorExperimentalSettings.h
@@ -132,20 +132,15 @@ public:
 	UPROPERTY(EditAnywhere, config, Category = "Content Browser", meta = (DisplayName = "Enable support for long paths (> 260 characters)"))
 	bool bEnableLongPathsSupport;
 
-	/** Allows creating APackedLevelActor blueprint actors */
-	UPROPERTY(EditAnywhere, config, Category = Level)
-	bool bPackedLevelActor;
-
-	/** Allows creating ALevelInstance actors */
-	UPROPERTY(EditAnywhere, config, Category = Level)
-	bool bLevelInstance;
-
 	UPROPERTY(EditAnywhere, config, Category = WorldPartition)
 	bool bEnableWorldPartitionActorFilters;
 
 	UPROPERTY(EditAnywhere, config, Category = WorldPartition)
 	bool bEnableWorldPartitionExternalDataLayers;
 
+	UPROPERTY(EditAnywhere, config, Category = LevelInstance)
+	bool bEnableLevelInstancePropertyOverrides;
+
 	/**
 	 * Returns an event delegate that is executed when a setting has changed.
 	 *
diff --git a/Engine/Source/Editor/UnrealEd/Private/Settings/SettingsClasses.cpp b/Engine/Source/Editor/UnrealEd/Private/Settings/SettingsClasses.cpp
--- a/Engine/Source/Editor/UnrealEd/Private/Settings/SettingsClasses.cpp
+++ b/Engine/Source/Editor/UnrealEd/Private/Settings/SettingsClasses.cpp
@@ -192,8 +192,6 @@ UEditorExperimentalSettings::UEditorExperimentalSettings( const FObjectInitializ
 	, HDREditorNITLevel(160.0f)
 	, bUseOpenCLForConvexHullDecomp(false)
 	, bAllowPotentiallyUnsafePropertyEditing(false)
-	, bPackedLevelActor(true)
-	, bLevelInstance(true)
 {
 }
 
diff --git a/Engine/Source/Runtime/Core/Public/UObject/FortniteMainBranchObjectVersion.h b/Engine/Source/Runtime/Core/Public/UObject/FortniteMainBranchObjectVersion.h
--- a/Engine/Source/Runtime/Core/Public/UObject/FortniteMainBranchObjectVersion.h
+++ b/Engine/Source/Runtime/Core/Public/UObject/FortniteMainBranchObjectVersion.h
@@ -426,6 +426,9 @@ struct FFortniteMainBranchObjectVersion
 		// Move workspaces to a seperate plugin
 		AnimNextMoveWorkspaces,
 
+		// Level Instance Property overrides
+		LevelInstancePropertyOverrides,
+
 		// -----<new versions can be added above this line>-------------------------------------------------
 		VersionPlusOne,
 		LatestVersion = VersionPlusOne - 1
diff --git a/Engine/Source/Runtime/Engine/Classes/GameFramework/Actor.h b/Engine/Source/Runtime/Engine/Classes/GameFramework/Actor.h
--- a/Engine/Source/Runtime/Engine/Classes/GameFramework/Actor.h
+++ b/Engine/Source/Runtime/Engine/Classes/GameFramework/Actor.h
@@ -966,7 +966,9 @@ protected:
 	UPROPERTY(VisibleAnywhere, AdvancedDisplay, Category = DataLayers)
 	TArray<FActorDataLayer> DataLayers;
 
-	UPROPERTY(EditAnywhere, AdvancedDisplay, Category = DataLayers)
+	// There is currently an issue where if we allow property override of DataLayerAssets and it contains Private datalayers
+	// then it will always serialize a diff since those are outered to the instanced level and will get remapped differently between the Override instance and Archetype instance we are comparing against
+	UPROPERTY(EditAnywhere, AdvancedDisplay, Category = DataLayers, meta=(DisableLevelInstancePropertyOverride))
 	TArray<TSoftObjectPtr<UDataLayerAsset>> DataLayerAssets;
 
 	TArray<TSoftObjectPtr<UDataLayerAsset>> PreEditChangeDataLayers;
@@ -4514,6 +4516,7 @@ private:
 
 	friend class ULevelStreamingLevelInstance;
 	friend class ULevelStreamingLevelInstanceEditor;
+	friend class ULevelStreamingLevelInstanceEditorPropertyOverride;
 };
 #endif
 
diff --git a/Engine/Source/Runtime/Engine/Private/LevelInstance/LevelInstanceActor.cpp b/Engine/Source/Runtime/Engine/Private/LevelInstance/LevelInstanceActor.cpp
--- a/Engine/Source/Runtime/Engine/Private/LevelInstance/LevelInstanceActor.cpp
+++ b/Engine/Source/Runtime/Engine/Private/LevelInstance/LevelInstanceActor.cpp
@@ -99,7 +99,7 @@ bool ALevelInstance::IsLoadingEnabled() const
 const TSoftObjectPtr<UWorld>& ALevelInstance::GetWorldAsset() const
 {
 #if WITH_EDITORONLY_DATA
-	return WorldAsset;
+	return GetPropertyOverrideAsset() ? GetPropertyOverrideAsset()->GetWorldAsset() : WorldAsset;
 #else
 	return CookedWorldAsset;
 #endif
@@ -126,6 +126,17 @@ void ALevelInstance::OnLevelInstanceLoaded()
 }
 
 #if WITH_EDITOR
+
+void ALevelInstance::SetPropertyOverrideAsset(ULevelInstancePropertyOverrideAsset* InPropertyOverrideAsset)
+{
+	if (PropertyOverrides != InPropertyOverrideAsset)
+	{
+		Modify();
+		PropertyOverrides = InPropertyOverrideAsset;
+	}
+}
+
+
 ULevelInstanceComponent* ALevelInstance::GetLevelInstanceComponent() const
 {
 	return Cast<ULevelInstanceComponent>(RootComponent);
@@ -141,6 +152,16 @@ bool ALevelInstance::SupportsPartialEditorLoading() const
 	return ILevelInstanceInterface::SupportsPartialEditorLoading() && LevelInstanceActorImpl.SupportsPartialEditorLoading();
 }
 
+bool ALevelInstance::SupportsPropertyOverrides() const
+{
+	return LevelInstanceActorImpl.SupportsPropertyOverrides();
+}
+
+ULevelInstancePropertyOverrideAsset* ALevelInstance::GetPropertyOverrideAsset() const
+{
+	return SupportsPropertyOverrides() ? PropertyOverrides : nullptr;
+}
+
 TUniquePtr<FWorldPartitionActorDesc> ALevelInstance::CreateClassActorDesc() const
 {
 	return TUniquePtr<FWorldPartitionActorDesc>(new FLevelInstanceActorDesc());
@@ -158,7 +179,7 @@ void ALevelInstance::PostLoad()
 #if WITH_EDITORONLY_DATA
 	if (IsRunningCookCommandlet() && ShouldCookWorldAsset())
 	{
-		CookedWorldAsset = WorldAsset;
+		CookedWorldAsset = GetWorldAsset();
 	}
 #endif
 }
@@ -300,9 +321,9 @@ bool ALevelInstance::GetReferencedContentObjects(TArray<UObject*>& Objects) cons
 
 bool ALevelInstance::GetSoftReferencedContentObjects(TArray<FSoftObjectPath>& SoftObjects) const
 {
-	if (WorldAsset.ToSoftObjectPath().IsValid())
+	if (GetWorldAsset().ToSoftObjectPath().IsValid())
 	{
-		SoftObjects.Add(WorldAsset.ToSoftObjectPath());
+		SoftObjects.Add(GetWorldAsset().ToSoftObjectPath());
 		return true;
 	}
 	return false;
diff --git a/Engine/Source/Runtime/Engine/Private/LevelInstance/LevelInstanceActorImpl.cpp b/Engine/Source/Runtime/Engine/Private/LevelInstance/LevelInstanceActorImpl.cpp
--- a/Engine/Source/Runtime/Engine/Private/LevelInstance/LevelInstanceActorImpl.cpp
+++ b/Engine/Source/Runtime/Engine/Private/LevelInstance/LevelInstanceActorImpl.cpp
@@ -9,6 +9,7 @@
 #include "LevelInstance/LevelInstancePrivate.h"
 #include "LevelInstance/LevelInstanceComponent.h"
 #include "LevelInstance/LevelInstanceEditorInstanceActor.h"
+#include "LevelInstance/LevelInstanceSettings.h"
 #include "WorldPartition/WorldPartitionActorLoaderInterface.h"
 #include "Engine/Level.h"
 #include "Engine/World.h"
@@ -116,6 +117,11 @@ bool FLevelInstanceActorImpl::ResolveSubobject(const TCHAR* SubObjectPath, UObje
 	return false;
 }
 
+bool FLevelInstanceActorImpl::SupportsPropertyOverrides() const
+{
+	return ULevelInstanceSettings::Get()->IsPropertyOverrideEnabled();
+}
+
 bool FLevelInstanceActorImpl::SupportsPartialEditorLoading() const
 {
 	return bAllowPartialLoading;
@@ -315,7 +321,7 @@ bool FLevelInstanceActorImpl::IsLockedActor() const
 		return true;
 	}
 
-	if (LevelInstance->IsEditing())
+	if (LevelInstance->IsEditing() || LevelInstance->IsEditingPropertyOverrides())
 	{
 		return true;
 	}
diff --git a/Engine/Source/Runtime/Engine/Private/LevelInstance/LevelInstanceComponent.cpp b/Engine/Source/Runtime/Engine/Private/LevelInstance/LevelInstanceComponent.cpp
--- a/Engine/Source/Runtime/Engine/Private/LevelInstance/LevelInstanceComponent.cpp
+++ b/Engine/Source/Runtime/Engine/Private/LevelInstance/LevelInstanceComponent.cpp
@@ -208,6 +208,11 @@ void ULevelInstanceComponent::OnEdit()
 	}
 }
 
+void ULevelInstanceComponent::OnEditOverrides()
+{
+	OnEdit();
+}
+
 void ULevelInstanceComponent::OnCommit()
 {
 	if (SpriteComponent)
@@ -216,6 +221,11 @@ void ULevelInstanceComponent::OnCommit()
 	}
 }
 
+void ULevelInstanceComponent::OnCommitOverrides()
+{
+	OnCommit();
+}
+
 void ULevelInstanceComponent::UpdateEditFilter()
 {
 	// @todo_ow: for child level instances setup the proper filter based on parent hierarchy filter
diff --git a/Engine/Source/Runtime/Engine/Private/LevelInstance/LevelInstanceEditorLevelStreaming.cpp b/Engine/Source/Runtime/Engine/Private/LevelInstance/LevelInstanceEditorLevelStreaming.cpp
--- a/Engine/Source/Runtime/Engine/Private/LevelInstance/LevelInstanceEditorLevelStreaming.cpp
+++ b/Engine/Source/Runtime/Engine/Private/LevelInstance/LevelInstanceEditorLevelStreaming.cpp
@@ -140,7 +140,12 @@ void ULevelStreamingLevelInstanceEditor::Unload(ULevelStreamingLevelInstanceEdit
 		if (ULevel* LoadedLevel = LevelStreaming->GetLoadedLevel())
 		{
 			LoadedLevel->OnLoadedActorAddedToLevelPreEvent.RemoveAll(LevelStreaming);
-			LevelInstanceSubsystem->RemoveLevelsFromWorld({ LoadedLevel });
+			
+			// Not needed if world is being cleaned up
+			if (!LevelStreaming->GetWorld()->IsBeingCleanedUp())
+			{
+				LevelInstanceSubsystem->RemoveLevelsFromWorld({ LoadedLevel });
+			}
 		}
 	}
 }
diff --git a/Engine/Source/Runtime/Engine/Private/LevelInstance/LevelInstanceInterface.cpp b/Engine/Source/Runtime/Engine/Private/LevelInstance/LevelInstanceInterface.cpp
--- a/Engine/Source/Runtime/Engine/Private/LevelInstance/LevelInstanceInterface.cpp
+++ b/Engine/Source/Runtime/Engine/Private/LevelInstance/LevelInstanceInterface.cpp
@@ -26,7 +26,7 @@ bool ILevelInstanceInterface::SupportsPartialEditorLoading() const
 	{
 		if (ULevelInstanceSubsystem* LevelInstanceSubsystem = GetLevelInstanceSubsystem())
 		{
-			if (LevelInstanceSubsystem->IsEditingLevelInstance(this) || HasChildEdit())
+			if (LevelInstanceSubsystem->IsEditingLevelInstance(this) || LevelInstanceSubsystem->IsEditingLevelInstancePropertyOverrides(this) || LevelInstanceSubsystem->HasChildEdit(this))
 			{
 				return false;
 			}
@@ -185,6 +185,18 @@ bool ILevelInstanceInterface::IsEditing() const
 	return false;
 }
 
+bool ILevelInstanceInterface::IsEditingPropertyOverrides() const
+{
+	// No need to check LevelInstanceID here since Override Level Instance is referenced directly:
+	// Level Instance ID can become invalid while Actor is being edited through AActor property change events causing component unregisters
+	if (ULevelInstanceSubsystem* LevelInstanceSubsystem = GetLevelInstanceSubsystem())
+	{
+		return LevelInstanceSubsystem->IsEditingLevelInstancePropertyOverrides(this);
+	}
+
+	return false;
+}
+
 bool ILevelInstanceInterface::HasChildEdit() const
 {
 	if (HasValidLevelInstanceID())
@@ -237,6 +249,19 @@ bool ILevelInstanceInterface::HasDirtyChildren() const
 	return false;
 }
 
+bool ILevelInstanceInterface::CanEnterEditPropertyOverrides(FText* OutReason) const
+{
+	if (HasValidLevelInstanceID())
+	{
+		if (ULevelInstanceSubsystem* LevelInstanceSubsystem = GetLevelInstanceSubsystem())
+		{
+			return LevelInstanceSubsystem->CanEditLevelInstancePropertyOverrides(this, OutReason);
+		}
+	}
+
+	return false;
+}
+
 bool ILevelInstanceInterface::CanEnterEdit(FText* OutReason) const
 {
 	if (HasValidLevelInstanceID())
@@ -250,6 +275,14 @@ bool ILevelInstanceInterface::CanEnterEdit(FText* OutReason) const
 	return false;
 }
 
+bool ILevelInstanceInterface::EnterEditPropertyOverrides(AActor* ContextActor)
+{
+	ULevelInstanceSubsystem* LevelInstanceSubsystem = GetLevelInstanceSubsystem();
+	check(LevelInstanceSubsystem);
+	LevelInstanceSubsystem->EditLevelInstancePropertyOverrides(this, ContextActor);
+	return true;
+}
+
 bool ILevelInstanceInterface::EnterEdit(AActor* ContextActor)
 {
 	ULevelInstanceSubsystem* LevelInstanceSubsystem = GetLevelInstanceSubsystem();
@@ -266,6 +299,27 @@ void ILevelInstanceInterface::OnEdit()
 	}
 }
 
+void ILevelInstanceInterface::OnEditPropertyOverrides()
+{
+	if (ULevelInstanceComponent* LevelInstanceComponent = GetLevelInstanceComponent())
+	{
+		LevelInstanceComponent->OnEditOverrides();
+	}
+}
+
+bool ILevelInstanceInterface::CanExitEditPropertyOverrides(bool bDiscardEdits, FText* OutReason) const
+{
+	if (HasValidLevelInstanceID())
+	{
+		if (ULevelInstanceSubsystem* LevelInstanceSubsystem = GetLevelInstanceSubsystem())
+		{
+			return LevelInstanceSubsystem->CanCommitLevelInstancePropertyOverrides(this, bDiscardEdits, OutReason);
+		}
+	}
+
+	return false;
+}
+
 bool ILevelInstanceInterface::CanExitEdit(bool bDiscardEdits, FText* OutReason) const
 {
 	if (HasValidLevelInstanceID())
@@ -279,6 +333,14 @@ bool ILevelInstanceInterface::CanExitEdit(bool bDiscardEdits, FText* OutReason)
 	return false;
 }
 
+bool ILevelInstanceInterface::ExitEditPropertyOverrides(bool bDiscardEdits)
+{
+	ULevelInstanceSubsystem* LevelInstanceSubsystem = GetLevelInstanceSubsystem();
+	check(LevelInstanceSubsystem);
+	return LevelInstanceSubsystem->CommitLevelInstancePropertyOverrides(this, bDiscardEdits);
+}
+
+
 bool ILevelInstanceInterface::ExitEdit(bool bDiscardEdits)
 {
 	ULevelInstanceSubsystem* LevelInstanceSubsystem = GetLevelInstanceSubsystem();
@@ -286,6 +348,14 @@ bool ILevelInstanceInterface::ExitEdit(bool bDiscardEdits)
 	return LevelInstanceSubsystem->CommitLevelInstance(this, bDiscardEdits);
 }
 
+void ILevelInstanceInterface::OnCommitPropertyOverrides(bool bChanged)
+{
+	if (ULevelInstanceComponent* LevelInstanceComponent = GetLevelInstanceComponent())
+	{
+		LevelInstanceComponent->OnCommitOverrides();
+	}
+}
+
 void ILevelInstanceInterface::OnCommit(bool bChanged)
 {
 	if (ULevelInstanceComponent* LevelInstanceComponent = GetLevelInstanceComponent())
diff --git a/Engine/Source/Runtime/Engine/Private/LevelInstance/LevelInstanceLevelStreaming.cpp b/Engine/Source/Runtime/Engine/Private/LevelInstance/LevelInstanceLevelStreaming.cpp
--- a/Engine/Source/Runtime/Engine/Private/LevelInstance/LevelInstanceLevelStreaming.cpp
+++ b/Engine/Source/Runtime/Engine/Private/LevelInstance/LevelInstanceLevelStreaming.cpp
@@ -21,6 +21,8 @@
 
 #if WITH_EDITOR
 #include "LevelInstance/LevelInstanceEditorInstanceActor.h"
+#include "WorldPartition/LevelInstance/LevelInstanceContainerInstance.h"
+#include "LevelInstance/LevelInstancePropertyOverrideAsset.h"
 #include "LevelUtils.h"
 #include "ActorFolder.h"
 #include "Misc/LazySingleton.h"
@@ -134,12 +136,95 @@ FBox ULevelStreamingLevelInstance::GetBounds() const
 	return CachedBounds;
 }
 
+void ULevelStreamingLevelInstance::OnActorReplacedEvent(FWorldPartitionActorDescInstance* InActorDescInstance)
+{
+	if (ILevelInstanceInterface* LevelInstance = GetLevelInstance())
+	{
+		ULevelInstanceSubsystem* LevelInstanceSubsystem = LevelInstance->GetLevelInstanceSubsystem();
+		check(LevelInstanceSubsystem);
+		if (AActor* Actor = InActorDescInstance->GetActor())
+		{
+			ApplyPropertyOverrides({ Actor }, true, EApplyPropertieOverrideType::PreAndPostConstruction);
+		}
+	}
+}
+
 void ULevelStreamingLevelInstance::OnLoadedActorsAddedToLevelPreEvent(const TArray<AActor*>& InActors)
+{
+	InitializeActors(InActors);
+	ApplyPropertyOverrides(InActors, false, EApplyPropertieOverrideType::PreConstruction);
+}
+
+void ULevelStreamingLevelInstance::ApplyPropertyOverrides(const TArray<AActor*>& InActors, bool bInAlreadyAppliedTransformOnActors, EApplyPropertieOverrideType InApplyPropertyOverrideType)
 {
 	if (IsEditorWorldMode())
 	{
 		if (ILevelInstanceInterface* LevelInstance = GetLevelInstance())
 		{
+			ULevelInstanceSubsystem* LevelInstanceSubsystem = LevelInstance->GetLevelInstanceSubsystem();
+			check(LevelInstanceSubsystem);
+
+			const FActorContainerID ContainerID = LevelInstance->GetLevelInstanceID().GetContainerID();
+			const FActorContainerID ContextContainerID = LevelInstanceSubsystem->GetLevelInstancePropertyOverridesContext(LevelInstance);
+
+			for (AActor* Actor : InActors)
+			{
+				if (IsValid(Actor))
+				{
+					// Gather Contextual Property Overrides and apply them
+					TArray<const FActorPropertyOverride*> ActorPropertyOverrides;
+					if (LevelInstanceSubsystem->GetLevelInstancePropertyOverridesForActor(Actor, ContextContainerID, ActorPropertyOverrides))
+					{
+						// If we have Property Overrides we need to Remove the level transform before applying them in case the Relative transform of the actors was modified
+						if (bInAlreadyAppliedTransformOnActors && Actor->GetRootComponent())
+						{
+							FLevelUtils::RemoveEditorTransform(this, false, Actor);
+						}
+
+						bool bAppliedProperties = false;
+						if (InApplyPropertyOverrideType == EApplyPropertieOverrideType::PreConstruction || InApplyPropertyOverrideType == EApplyPropertieOverrideType::PreAndPostConstruction)
+						{
+							for (const FActorPropertyOverride* ActorPropertyOverride : ActorPropertyOverrides)
+							{
+								bAppliedProperties |= ULevelInstancePropertyOverrideAsset::ApplyPropertyOverrides(ActorPropertyOverride, Actor, false);
+							}
+						}
+
+						// If we did apply some properties re-run construction script on actor
+						if (bAppliedProperties && InApplyPropertyOverrideType == EApplyPropertieOverrideType::PreAndPostConstruction)
+						{
+							Actor->RerunConstructionScripts();
+						}
+
+						if (InApplyPropertyOverrideType == EApplyPropertieOverrideType::PostConstruction || InApplyPropertyOverrideType == EApplyPropertieOverrideType::PreAndPostConstruction)
+						{
+							for (const FActorPropertyOverride* ActorPropertyOverride : ActorPropertyOverrides)
+							{
+								ULevelInstancePropertyOverrideAsset::ApplyPropertyOverrides(ActorPropertyOverride, Actor, true);
+							}
+						}
+						
+						if (bInAlreadyAppliedTransformOnActors && Actor->GetRootComponent())
+						{
+							FLevelUtils::ApplyEditorTransform(this, true, Actor);
+							Actor->GetRootComponent()->UpdateComponentToWorld();
+							Actor->MarkComponentsRenderStateDirty();
+						}
+					}
+				}
+			}
+		}
+	}
+}
+
+void ULevelStreamingLevelInstance::InitializeActors(const TArray<AActor*>& InActors)
+{
+	if (IsEditorWorldMode())
+	{
+		if (ILevelInstanceInterface* LevelInstance = GetLevelInstance())
+		{
+			const FActorContainerID ContainerID = LevelInstance->GetLevelInstanceID().GetContainerID();
+		
 			for (AActor* Actor : InActors)
 			{
 				if (IsValid(Actor))
@@ -162,7 +247,6 @@ void ULevelStreamingLevelInstance::OnLoadedActorsAddedToLevelPreEvent(const TArr
 					}
 
 					// Must happen before the actors are registered with the world, which is the case for this delegate.
-					const FActorContainerID& ContainerID = LevelInstance->GetLevelInstanceID().GetContainerID();
 					FSetActorInstanceGuid SetActorInstanceGuid(Actor, ContainerID.GetActorGuid(Actor->GetActorGuid()));
 
 					FSetActorIsInLevelInstance SetIsInLevelInstance(Actor);
@@ -200,6 +284,8 @@ void ULevelStreamingLevelInstance::OnLoadedActorsAddedToLevelPostEvent(const TAr
 					}
 				}
 			}
+
+			ApplyPropertyOverrides(InActors, true, EApplyPropertieOverrideType::PostConstruction);
 		}
 	}
 }
@@ -264,8 +350,17 @@ void ULevelStreamingLevelInstance::OnPreInitializeContainerInstance(UActorDescCo
 	
 	// In Editor it is possible to have a non WP parent world in which case we pass in null to the SetParent method, this will ensure that in editor the Level Instance container ID won't be a IsMainContainer() and will properly handle IsMainWorldOnly actors
 	InInitParams.SetParent(OwningWorldPartition ? OwningWorldPartition->GetActorDescContainerInstance() : nullptr, LevelInstanceActor->GetActorGuid());
+		
+	// Apply Override Container
+	if (OwningWorldPartition)
+	{
+		if (FWorldPartitionActorDescInstance* LevelInstanceActorDescInstance = OwningWorldPartition->GetActorDescInstance(LevelInstanceActor->GetActorGuid()); LevelInstanceActorDescInstance && LevelInstanceActorDescInstance->IsChildContainerInstance())
+		{
+			ULevelInstanceContainerInstance* LevelInstanceContainerInstance = CastChecked<ULevelInstanceContainerInstance>(InContainerInstance);
+			LevelInstanceContainerInstance->SetOverrideContainerAndAsset(LevelInstanceActorDescInstance->GetActorDesc()->GetChildContainer(), GetLevelInstance()->GetPropertyOverrideAsset());
+		}
+	}
 }
-
 #endif
 
 ULevelStreamingLevelInstance* ULevelStreamingLevelInstance::LoadInstance(ILevelInstanceInterface* LevelInstance)
@@ -330,9 +425,10 @@ ULevelStreamingLevelInstance* ULevelStreamingLevelInstance::LoadInstance(ILevelI
 }
 
 #if WITH_EDITOR
-void ULevelStreamingLevelInstance::OnLevelStreamingStateChanged(UWorld* InWorld, const ULevelStreaming* InLevelStreaming, ULevel* InLevelIfLoaded, ELevelStreamingState InPrevState, ELevelStreamingState InNewState)
+void ULevelStreamingLevelInstance::OnCurrentStateChanged(ELevelStreamingState InPrevState, ELevelStreamingState InNewState)
 {
-	if (InNewState == ELevelStreamingState::LoadedVisible && InLevelStreaming == this)
+	Super::OnCurrentStateChanged(InPrevState, InNewState);
+	if (InNewState == ELevelStreamingState::LoadedVisible && IsEditorWorldMode())
 	{
 		ILevelInstanceInterface* LevelInstance = GetLevelInstance();
 
@@ -351,9 +447,11 @@ void ULevelStreamingLevelInstance::OnLevelStreamingStateChanged(UWorld* InWorld,
 			}
 		}, /*bIncludeNestedObjects*/ true);
 
-		OnLoadedActorsAddedToLevelPreEvent(Level->Actors);
+		// Initialize Deferred (Async Actors)s
+		check(Level->bAlreadyMovedActors);
+		ApplyPropertyOverrides(Level->Actors, true, EApplyPropertieOverrideType::PostConstruction);
 
-		Level->OnLoadedActorAddedToLevelPreEvent.AddUObject(this, &ULevelStreamingLevelInstance::OnLoadedActorsAddedToLevelPreEvent);
+		// Currently only used for Attachement which doesn't do anything before the ALevelInstanceEditorInstanceActor::Create call that follows
 		Level->OnLoadedActorAddedToLevelPostEvent.AddUObject(this, &ULevelStreamingLevelInstance::OnLoadedActorsAddedToLevelPostEvent);
 		Level->OnLoadedActorRemovedFromLevelPreEvent.AddUObject(this, &ULevelStreamingLevelInstance::OnLoadedActorsRemovedFromLevelPostEvent);
 
@@ -363,9 +461,6 @@ void ULevelStreamingLevelInstance::OnLevelStreamingStateChanged(UWorld* InWorld,
 		// Push editing state to child actors
 		AActor* LevelInstanceActor = CastChecked<AActor>(LevelInstance);
 		LevelInstanceActor->PushLevelInstanceEditingStateToProxies(LevelInstanceActor->IsInEditLevelInstanceHierarchy());
-
-		// Unregister
-		FLevelStreamingDelegates::OnLevelStreamingStateChanged.RemoveAll(this);
 	}
 }
 #endif
@@ -375,14 +470,17 @@ void ULevelStreamingLevelInstance::UnloadInstance(ULevelStreamingLevelInstance*
 #if WITH_EDITOR
 	if (LevelStreaming->IsEditorWorldMode())
 	{
-		FLevelStreamingDelegates::OnLevelStreamingStateChanged.RemoveAll(LevelStreaming);
-
 		ULevel* LoadedLevel = LevelStreaming->GetLoadedLevel();
 		LoadedLevel->OnLoadedActorAddedToLevelPreEvent.RemoveAll(LevelStreaming);
 		LoadedLevel->OnLoadedActorAddedToLevelPostEvent.RemoveAll(LevelStreaming);
 		LoadedLevel->OnLoadedActorRemovedFromLevelPreEvent.RemoveAll(LevelStreaming);
 		LevelStreaming->LevelInstanceEditorInstanceActor.Reset();
 
+		if (UWorldPartition* OuterWorldPartition = LoadedLevel->GetWorldPartition())
+		{
+			OuterWorldPartition->OnActorReplacedEvent.RemoveAll(LevelStreaming);
+		}
+
 		// Check if we need to flush the Trans buffer...
 		UWorld* OuterWorld = LoadedLevel->GetTypedOuter<UWorld>();
 		bool bResetTrans = false;
@@ -418,18 +516,18 @@ void ULevelStreamingLevelInstance::OnLevelLoadedChanged(ULevel* InLevel)
 	{
 #if WITH_EDITOR
 		if (IsEditorWorldMode())
-		{
-			FLevelStreamingDelegates::OnLevelStreamingStateChanged.AddUObject(this, &ULevelStreamingLevelInstance::OnLevelStreamingStateChanged);
+		{	
+			InitializeActors(InLevel->Actors);
+			ApplyPropertyOverrides(InLevel->Actors, false, EApplyPropertieOverrideType::PreConstruction);
+			NewLoadedLevel->OnLoadedActorAddedToLevelPreEvent.AddUObject(this, &ULevelStreamingLevelInstance::OnLoadedActorsAddedToLevelPreEvent);
 		}
 #endif
-
 		check(InLevel == NewLoadedLevel);
-		if (!NewLoadedLevel->bAlreadyMovedActors)
-		{
-			AWorldSettings* WorldSettings = NewLoadedLevel->GetWorldSettings();
-			check(WorldSettings);
-			LevelTransform = FTransform(WorldSettings->LevelInstancePivotOffset) * LevelTransform;
-		}
+		check(!NewLoadedLevel->bAlreadyMovedActors);
+		
+		AWorldSettings* WorldSettings = NewLoadedLevel->GetWorldSettings();
+		check(WorldSettings);
+		LevelTransform = FTransform(WorldSettings->LevelInstancePivotOffset) * LevelTransform;
 
 		if (ULevelInstanceSubsystem* LevelInstanceSubsystem = GetWorld()->GetSubsystem<ULevelInstanceSubsystem>())
 		{
@@ -445,6 +543,8 @@ void ULevelStreamingLevelInstance::OnLevelLoadedChanged(ULevel* InLevel)
 				if (IsEditorWorldMode())
 				{
 					OuterWorldPartition->OnActorDescContainerInstancePreInitialize.BindUObject(this, &ULevelStreamingLevelInstance::OnPreInitializeContainerInstance);
+					OuterWorldPartition->SetContainerInstanceClass(ULevelInstanceContainerInstance::StaticClass());
+					OuterWorldPartition->OnActorReplacedEvent.AddUObject(this, &ULevelStreamingLevelInstance::OnActorReplacedEvent);
 				}
 
 				if (UWorldPartition* OwningWorldPartition = GetWorld()->GetWorldPartition(); OwningWorldPartition && OwningWorldPartition->IsStreamingEnabled())
