Don't call UpdateEditorMode while running a commandlet

diff --git a/Engine/Source/Editor/LevelInstanceEditor/Private/LevelInstanceEditorModule.cpp b/Engine/Source/Editor/LevelInstanceEditor/Private/LevelInstanceEditorModule.cpp
--- a/Engine/Source/Editor/LevelInstanceEditor/Private/LevelInstanceEditorModule.cpp
+++ b/Engine/Source/Editor/LevelInstanceEditor/Private/LevelInstanceEditorModule.cpp
@@ -993,13 +993,16 @@ void FLevelInstanceEditorModule::BroadcastTryExitEditorMode()
 
 void FLevelInstanceEditorModule::UpdateEditorMode(bool bActivated)
 {
-	if (bActivated && !GLevelEditorModeTools().IsModeActive(ULevelInstanceEditorMode::EM_LevelInstanceEditorModeId))
-	{
-		GLevelEditorModeTools().ActivateMode(ULevelInstanceEditorMode::EM_LevelInstanceEditorModeId);
-	}
-	else if(!bActivated && GLevelEditorModeTools().IsModeActive(ULevelInstanceEditorMode::EM_LevelInstanceEditorModeId))
+	if (!IsRunningCommandlet() && GLevelEditorModeToolsIsValid())
 	{
-		GLevelEditorModeTools().DeactivateMode(ULevelInstanceEditorMode::EM_LevelInstanceEditorModeId);
+		if (bActivated && !GLevelEditorModeTools().IsModeActive(ULevelInstanceEditorMode::EM_LevelInstanceEditorModeId))
+		{
+			GLevelEditorModeTools().ActivateMode(ULevelInstanceEditorMode::EM_LevelInstanceEditorModeId);
+		}
+		else if (!bActivated && GLevelEditorModeTools().IsModeActive(ULevelInstanceEditorMode::EM_LevelInstanceEditorModeId))
+		{
+			GLevelEditorModeTools().DeactivateMode(ULevelInstanceEditorMode::EM_LevelInstanceEditorModeId);
+		}
 	}
 }
 
