Fix conflicting localization

#rnx

diff --git a/Engine/Source/Editor/LevelInstanceEditor/Private/LevelInstanceEditorModule.cpp b/Engine/Source/Editor/LevelInstanceEditor/Private/LevelInstanceEditorModule.cpp
--- a/Engine/Source/Editor/LevelInstanceEditor/Private/LevelInstanceEditorModule.cpp
+++ b/Engine/Source/Editor/LevelInstanceEditor/Private/LevelInstanceEditorModule.cpp
@@ -117,10 +117,10 @@ struct FLevelInstanceMenuUtils
 			return bCanEdit;
 		});
 
-		FText EntryLabel = bSingleEntry ? LOCTEXT("EditLevelInstances", "Override") : FText::FromString(LevelInstance->GetWorldAsset().GetAssetName());
+		FText EntryLabel = bSingleEntry ? LOCTEXT("OverrideLevelInstances", "Override") : FText::FromString(LevelInstance->GetWorldAsset().GetAssetName());
 		if (bCanEdit)
 		{
-			EntryDesc = FText::Format(LOCTEXT("LevelInstanceName", "{0}:{1}"), FText::FromString(LevelInstanceActor->GetActorLabel()), FText::FromString(LevelInstance->GetWorldAssetPackage()));
+			EntryDesc = FText::Format(LOCTEXT("OverrideLevelInstanceName", "{0}:{1}"), FText::FromString(LevelInstanceActor->GetActorLabel()), FText::FromString(LevelInstance->GetWorldAssetPackage()));
 		}
 		Section.AddMenuEntry(NAME_None, EntryLabel, EntryDesc, FSlateIcon(), LevelInstanceEditAction);
 	}
