LevelEditor/ActorDetails Panel support for multiple ISCSEditorUICustomization
- Implement a private one in LevelInstanceEditorModule which acts as an array of ISCSEditorUICustomization
- First ISCSEditorUICustomization to return non default values is used for GetChildActorVisualizationMode() / GetComponentTypeFilter() (Sorted registration could be added later to prioritize a customization over others)
- Hide methods will return true if any of the Customizations returns true
- Convert existing code to new API
- Deprecate Hide methods for a new version taking a UObject arrayview context
- Deprecate Setter
- Remove Selection listening code in existing Customizations and use the context (fixes Locked details panel bugs)

#rnx

diff --git a/Engine/Plugins/Enterprise/DataprepEditor/Source/DataprepEditor/Private/Widgets/DataprepAssetView.cpp b/Engine/Plugins/Enterprise/DataprepEditor/Source/DataprepEditor/Private/Widgets/DataprepAssetView.cpp
--- a/Engine/Plugins/Enterprise/DataprepEditor/Source/DataprepEditor/Private/Widgets/DataprepAssetView.cpp
+++ b/Engine/Plugins/Enterprise/DataprepEditor/Source/DataprepEditor/Private/Widgets/DataprepAssetView.cpp
@@ -42,11 +42,11 @@ class FDataprepSCSEditorUICustomization : public ISCSEditorUICustomization
 public:
 	void SetHideComponentsTree(bool bInHide) { bHideComponentsTree = bInHide; }
 
-	virtual bool HideAddComponentButton() const override { return true; }
+	virtual bool HideAddComponentButton(TArrayView<UObject*> Context) const override { return true; }
 
-	virtual bool HideBlueprintButtons() const override { return true; }
+	virtual bool HideBlueprintButtons(TArrayView<UObject*> Context) const override { return true; }
 
-	virtual bool HideComponentsTree() const override { return bHideComponentsTree; }
+	virtual bool HideComponentsTree(TArrayView<UObject*> Context) const override { return bHideComponentsTree; }
 
 private:
 	bool bHideComponentsTree = false;
diff --git a/Engine/Plugins/Runtime/nDisplay/Source/DisplayClusterConfigurator/Private/DisplayClusterConfiguratorBlueprintEditor.cpp b/Engine/Plugins/Runtime/nDisplay/Source/DisplayClusterConfigurator/Private/DisplayClusterConfiguratorBlueprintEditor.cpp
--- a/Engine/Plugins/Runtime/nDisplay/Source/DisplayClusterConfigurator/Private/DisplayClusterConfiguratorBlueprintEditor.cpp
+++ b/Engine/Plugins/Runtime/nDisplay/Source/DisplayClusterConfigurator/Private/DisplayClusterConfiguratorBlueprintEditor.cpp
@@ -92,8 +92,8 @@ public:
 		return Instance.ToSharedRef();
 	}
 
-	virtual bool HideAddComponentButton() const override { return true; }
-	virtual bool HideBlueprintButtons() const override { return true; }
+	virtual bool HideAddComponentButton(TArrayView<UObject*> Context) const override { return true; }
+	virtual bool HideBlueprintButtons(TArrayView<UObject*> Context) const override { return true; }
 	virtual EChildActorComponentTreeViewVisualizationMode GetChildActorVisualizationMode() const override { return EChildActorComponentTreeViewVisualizationMode::UseDefault; }
 
 private:
diff --git a/Engine/Source/Editor/Kismet/Private/SSCSEditor.cpp b/Engine/Source/Editor/Kismet/Private/SSCSEditor.cpp
--- a/Engine/Source/Editor/Kismet/Private/SSCSEditor.cpp
+++ b/Engine/Source/Editor/Kismet/Private/SSCSEditor.cpp
@@ -6975,7 +6975,8 @@ void SSCSEditor::GetCollapsedNodes(const FSCSEditorTreeNodePtrType& InNodePtr, T
 
 EVisibility SSCSEditor::GetPromoteToBlueprintButtonVisibility() const
 {
-	return (UICustomization.IsValid() && UICustomization->HideBlueprintButtons())
+	UObject* ActorContextPtr = GetActorContext();
+	return (UICustomization.IsValid() && UICustomization->HideBlueprintButtons(MakeArrayView(&ActorContextPtr, 1)))
 		|| (EditorMode != EComponentEditorMode::ActorInstance) 
 		|| (GetBlueprint() != nullptr)
 		? EVisibility::Collapsed : EVisibility::Visible;
@@ -6983,7 +6984,8 @@ EVisibility SSCSEditor::GetPromoteToBlueprintButtonVisibility() const
 
 EVisibility SSCSEditor::GetEditBlueprintButtonVisibility() const
 {
-	return (UICustomization.IsValid() && UICustomization->HideBlueprintButtons())
+	UObject* ActorContextPtr = GetActorContext();
+	return (UICustomization.IsValid() && UICustomization->HideBlueprintButtons(MakeArrayView(&ActorContextPtr, 1)))
 		|| (EditorMode != EComponentEditorMode::ActorInstance)
 		|| (GetBlueprint() == nullptr)
 		? EVisibility::Collapsed : EVisibility::Visible;
@@ -6991,20 +6993,23 @@ EVisibility SSCSEditor::GetEditBlueprintButtonVisibility() const
 
 EVisibility SSCSEditor::GetComponentClassComboButtonVisibility() const
 {
+	UObject* ActorContextPtr = GetActorContext();
 	return (HideComponentClassCombo.Get() 
-		|| (UICustomization.IsValid() && UICustomization->HideAddComponentButton())) 
+		|| (UICustomization.IsValid() && UICustomization->HideAddComponentButton(MakeArrayView(&ActorContextPtr, 1))))
 		? EVisibility::Collapsed : EVisibility::Visible;
 }
 
 EVisibility SSCSEditor::GetComponentsTreeVisibility() const
 {
-	return (UICustomization.IsValid() && UICustomization->HideComponentsTree())
+	UObject* ActorContextPtr = GetActorContext();
+	return (UICustomization.IsValid() && UICustomization->HideComponentsTree(MakeArrayView(&ActorContextPtr, 1)))
 		? EVisibility::Collapsed : EVisibility::Visible;
 }
 
 EVisibility SSCSEditor::GetComponentsFilterBoxVisibility() const
 {
-	return (UICustomization.IsValid() && UICustomization->HideComponentsFilterBox())
+	UObject* ActorContextPtr = GetActorContext();
+	return (UICustomization.IsValid() && UICustomization->HideComponentsFilterBox(MakeArrayView(&ActorContextPtr, 1)))
 		? EVisibility::Collapsed : EVisibility::Visible;
 }
 
@@ -7387,7 +7392,8 @@ void SSCSEditor::SetUICustomization(TSharedPtr<ISCSEditorUICustomization> InUICu
 
 TSubclassOf<UActorComponent> SSCSEditor::GetComponentTypeFilterToApply() const
 {
-	TSubclassOf<UActorComponent> ComponentType = UICustomization.IsValid() ? UICustomization->GetComponentTypeFilter() : nullptr;
+	UObject* ActorContextPtr = GetActorContext();
+	TSubclassOf<UActorComponent> ComponentType = UICustomization.IsValid() ? UICustomization->GetComponentTypeFilter(MakeArrayView(&ActorContextPtr, 1)) : nullptr;
 	if (!ComponentType)
 	{
 		ComponentType = ComponentTypeFilter.Get();
diff --git a/Engine/Source/Editor/Kismet/Public/ISCSEditorUICustomization.h b/Engine/Source/Editor/Kismet/Public/ISCSEditorUICustomization.h
--- a/Engine/Source/Editor/Kismet/Public/ISCSEditorUICustomization.h
+++ b/Engine/Source/Editor/Kismet/Public/ISCSEditorUICustomization.h
@@ -5,6 +5,7 @@
 #include "Components/ChildActorComponent.h"
 #include "CoreMinimal.h"
 #include "Templates/SubclassOf.h"
+#include "Containers/ArrayView.h"
 
 class UActorComponent;
 
@@ -14,18 +15,33 @@ class ISCSEditorUICustomization
 public:
 	virtual ~ISCSEditorUICustomization() {}
 
-	/** @return Whether to hide the components tree */
+	UE_DEPRECATED(5.5, "Use version with ObjectContext instead")
 	virtual bool HideComponentsTree() const { return false; }
 
-	/** @return Whether to hide the components filter box */
+	UE_DEPRECATED(5.5, "Use version with ObjectContext instead")
 	virtual bool HideComponentsFilterBox() const { return false; }
 
-	/** @return Whether to hide the "Add Component" combo button */
+	UE_DEPRECATED(5.5, "Use version with ObjectContext instead")
 	virtual bool HideAddComponentButton() const { return false; }
 
-	/** @return Whether to hide the "Edit Blueprint" and "Blueprint/Add Script" buttons */
+	UE_DEPRECATED(5.5, "Use version with ObjectContext instead")
 	virtual bool HideBlueprintButtons() const { return false; }
 
+	UE_DEPRECATED(5.5, "Use version with ObjectContext instead")
+	virtual TSubclassOf<UActorComponent> GetComponentTypeFilter() const { return nullptr; }
+
+	/** @return Whether to hide the components tree */
+	virtual bool HideComponentsTree(TArrayView<UObject*> Context) const { return false; }
+
+	/** @return Whether to hide the components filter box */
+	virtual bool HideComponentsFilterBox(TArrayView<UObject*> Context) const { return false; }
+
+	/** @return Whether to hide the "Add Component" combo button */
+	virtual bool HideAddComponentButton(TArrayView<UObject*> Context) const { return false; }
+
+	/** @return Whether to hide the "Edit Blueprint" and "Blueprint/Add Script" buttons */
+	virtual bool HideBlueprintButtons(TArrayView<UObject*> Context) const { return false; }
+
 	/** 
 	 * @return An override for the default ChildActorComponentTreeViewVisualizationMode from the project settings, if different from UseDefault.
 	 * @note Setting an override also forces child actor tree view expansion to be enabled.
@@ -33,5 +49,5 @@ public:
 	virtual EChildActorComponentTreeViewVisualizationMode GetChildActorVisualizationMode() const { return EChildActorComponentTreeViewVisualizationMode::UseDefault; }
 
 	/** @return A component type that limits visible nodes when filtering the tree view */
-	virtual TSubclassOf<UActorComponent> GetComponentTypeFilter() const { return nullptr; }
+	virtual TSubclassOf<UActorComponent> GetComponentTypeFilter(TArrayView<UObject*> Context) const { return nullptr; }
 };
diff --git a/Engine/Source/Editor/LevelEditor/Private/ActorDetailsSCSEditorUICustomization.cpp b/Engine/Source/Editor/LevelEditor/Private/ActorDetailsSCSEditorUICustomization.cpp
new file mode 100644
--- /dev/null
+++ b/Engine/Source/Editor/LevelEditor/Private/ActorDetailsSCSEditorUICustomization.cpp
@@ -0,0 +1,104 @@
+// Copyright Epic Games, Inc. All Rights Reserved.
+#include "ActorDetailsSCSEditorUICustomization.h"
+
+TSharedPtr<FActorDetailsSCSEditorUICustomization> FActorDetailsSCSEditorUICustomization::Instance;
+
+TSharedPtr<FActorDetailsSCSEditorUICustomization> FActorDetailsSCSEditorUICustomization::GetInstance()
+{
+	if (!Instance)
+	{
+		Instance = MakeShareable(new FActorDetailsSCSEditorUICustomization());
+	}
+	return Instance;
+}
+
+bool FActorDetailsSCSEditorUICustomization::HideComponentsTree(TArrayView<UObject*> Context) const
+{
+	for (const TSharedPtr<ISCSEditorUICustomization>& Customization : Customizations)
+	{
+		if (Customization->HideComponentsTree(Context))
+		{
+			return true;
+		}
+	}
+
+	return false;
+}
+
+bool FActorDetailsSCSEditorUICustomization::HideComponentsFilterBox(TArrayView<UObject*> Context) const
+{
+	for (const TSharedPtr<ISCSEditorUICustomization>& Customization : Customizations)
+	{
+		if (Customization->HideComponentsFilterBox(Context))
+		{
+			return true;
+		}
+	}
+
+	return false;
+}
+
+bool FActorDetailsSCSEditorUICustomization::HideAddComponentButton(TArrayView<UObject*> Context) const
+{
+	for (const TSharedPtr<ISCSEditorUICustomization>& Customization : Customizations)
+	{
+		if (Customization->HideAddComponentButton(Context))
+		{
+			return true;
+		}
+	}
+
+	return false;
+}
+
+bool FActorDetailsSCSEditorUICustomization::HideBlueprintButtons(TArrayView<UObject*> Context) const
+{
+	for (const TSharedPtr<ISCSEditorUICustomization>& Customization : Customizations)
+	{
+		if (Customization->HideBlueprintButtons(Context))
+		{
+			return true;
+		}
+	}
+
+	return false;
+}
+
+EChildActorComponentTreeViewVisualizationMode FActorDetailsSCSEditorUICustomization::GetChildActorVisualizationMode() const
+{
+	for (const TSharedPtr<ISCSEditorUICustomization>& Customization : Customizations)
+	{
+		EChildActorComponentTreeViewVisualizationMode VisualizationMode = Customization->GetChildActorVisualizationMode();
+		if (VisualizationMode != EChildActorComponentTreeViewVisualizationMode::UseDefault)
+		{
+			return VisualizationMode;
+		}
+	}
+
+	return EChildActorComponentTreeViewVisualizationMode::UseDefault;
+}
+
+TSubclassOf<UActorComponent> FActorDetailsSCSEditorUICustomization::GetComponentTypeFilter(TArrayView<UObject*> Context) const
+{
+	for (const TSharedPtr<ISCSEditorUICustomization>& Customization : Customizations)
+	{
+		if(TSubclassOf<UActorComponent> SubClass = Customization->GetComponentTypeFilter(Context))
+		{
+			return SubClass;
+		}
+	}
+
+	return nullptr;
+}
+
+void FActorDetailsSCSEditorUICustomization::AddCustomization(TSharedPtr<ISCSEditorUICustomization> Customization)
+{
+	Customizations.AddUnique(Customization);
+}
+
+void FActorDetailsSCSEditorUICustomization::RemoveCustomization(TSharedPtr<ISCSEditorUICustomization> Customization)
+{
+	Customizations.Remove(Customization);
+}
+
+
diff --git a/Engine/Source/Editor/LevelEditor/Private/ActorDetailsSCSEditorUICustomization.h b/Engine/Source/Editor/LevelEditor/Private/ActorDetailsSCSEditorUICustomization.h
new file mode 100644
--- /dev/null
+++ b/Engine/Source/Editor/LevelEditor/Private/ActorDetailsSCSEditorUICustomization.h
@@ -0,0 +1,36 @@
+// Copyright Epic Games, Inc. All Rights Reserved.
+
+#pragma once
+
+#include "ISCSEditorUICustomization.h"
+#include "Containers/Array.h"
+#include "Templates/SharedPointer.h"
+
+class FActorDetailsSCSEditorUICustomization : public ISCSEditorUICustomization
+{
+public:
+	static TSharedPtr<FActorDetailsSCSEditorUICustomization> GetInstance();
+	
+	virtual ~FActorDetailsSCSEditorUICustomization() {}
+
+	virtual bool HideComponentsTree(TArrayView<UObject*> Context) const override;
+
+	virtual bool HideComponentsFilterBox(TArrayView<UObject*> Context) const override;
+
+	virtual bool HideAddComponentButton(TArrayView<UObject*> Context) const override;
+
+	virtual bool HideBlueprintButtons(TArrayView<UObject*> Context) const override;
+
+	virtual EChildActorComponentTreeViewVisualizationMode GetChildActorVisualizationMode() const override;
+
+	virtual TSubclassOf<UActorComponent> GetComponentTypeFilter(TArrayView<UObject*> Context) const override;
+
+	void AddCustomization(TSharedPtr<ISCSEditorUICustomization> Customization);
+	void RemoveCustomization(TSharedPtr<ISCSEditorUICustomization> Customization);
+
+private:
+	FActorDetailsSCSEditorUICustomization() {}
+
+	static TSharedPtr<FActorDetailsSCSEditorUICustomization> Instance;
+	TArray<TSharedPtr<ISCSEditorUICustomization>> Customizations;
+};
\ No newline at end of file
diff --git a/Engine/Source/Editor/LevelEditor/Private/SLevelEditor.cpp b/Engine/Source/Editor/LevelEditor/Private/SLevelEditor.cpp
--- a/Engine/Source/Editor/LevelEditor/Private/SLevelEditor.cpp
+++ b/Engine/Source/Editor/LevelEditor/Private/SLevelEditor.cpp
@@ -204,6 +204,8 @@ void SLevelEditor::RegisterMenus()
 
 void SLevelEditor::Construct( const SLevelEditor::FArguments& InArgs)
 {
+	ActorDetailsSCSEditorUICustomization = FActorDetailsSCSEditorUICustomization::GetInstance();
+	
 	// Important: We use raw bindings here because we are releasing our binding in our destructor (where a weak pointer would be invalid)
 	// It's imperative that our delegate is removed in the destructor for the level editor module to play nicely with reloading.
 
@@ -398,6 +400,8 @@ void SLevelEditor::ConstructTitleBarMessages()
 
 SLevelEditor::~SLevelEditor()
 {
+	ActorDetailsSCSEditorUICustomization = nullptr;
+
 	// We're going away now, so make sure all toolkits that are hosted within this level editor are shut down
 	FToolkitManager::Get().OnToolkitHostDestroyed( this );
 	HostedToolkits.Reset();
@@ -2193,13 +2197,17 @@ void SLevelEditor::SetActorDetailsRootCustomization(TSharedPtr<FDetailsViewObjec
 	}
 }
 
-void SLevelEditor::SetActorDetailsSCSEditorUICustomization(TSharedPtr<ISCSEditorUICustomization> InActorDetailsSCSEditorUICustomization)
+void SLevelEditor::AddActorDetailsSCSEditorUICustomization(TSharedPtr<ISCSEditorUICustomization> InActorDetailsSCSEditorUICustomization)
 {
-	ActorDetailsSCSEditorUICustomization = InActorDetailsSCSEditorUICustomization;
+	check(ActorDetailsSCSEditorUICustomization);
+	ActorDetailsSCSEditorUICustomization->AddCustomization(InActorDetailsSCSEditorUICustomization);
+}
 
-	for (TSharedRef<SActorDetails> ActorDetails : GetAllActorDetails())
+void SLevelEditor::RemoveActorDetailsSCSEditorUICustomization(TSharedPtr<ISCSEditorUICustomization> InActorDetailsSCSEditorUICustomization)
+{
+	if (ActorDetailsSCSEditorUICustomization)
 	{
-		ActorDetails->SetSubobjectEditorUICustomization(ActorDetailsSCSEditorUICustomization);
+		ActorDetailsSCSEditorUICustomization->RemoveCustomization(InActorDetailsSCSEditorUICustomization);
 	}
 }
 
diff --git a/Engine/Source/Editor/LevelEditor/Private/SLevelEditor.h b/Engine/Source/Editor/LevelEditor/Private/SLevelEditor.h
--- a/Engine/Source/Editor/LevelEditor/Private/SLevelEditor.h
+++ b/Engine/Source/Editor/LevelEditor/Private/SLevelEditor.h
@@ -16,6 +16,7 @@
 #include "ILevelEditor.h"
 #include "LevelViewportTabContent.h"
 #include "SLevelEditorToolBox.h"
+#include "ActorDetailsSCSEditorUICustomization.h"
 
 class IAssetEditorInstance;
 class IDetailsView;
@@ -129,7 +130,8 @@ public:
 	virtual UWorld* GetWorld() const override;
 	virtual TSharedRef<SWidget> CreateActorDetails( const FName TabIdentifier ) override;
 	virtual void SetActorDetailsRootCustomization(TSharedPtr<FDetailsViewObjectFilter> InActorDetailsObjectFilter, TSharedPtr<IDetailRootObjectCustomization> InActorDetailsRootCustomization) override;
-	virtual void SetActorDetailsSCSEditorUICustomization(TSharedPtr<ISCSEditorUICustomization> InActorDetailsSCSEditorUICustomization) override;
+	virtual void AddActorDetailsSCSEditorUICustomization(TSharedPtr<ISCSEditorUICustomization> ActorDetailsSCSEditorUICustomization) override;
+	virtual void RemoveActorDetailsSCSEditorUICustomization(TSharedPtr<ISCSEditorUICustomization> ActorDetailsSCSEditorUICustomization) override;
 	virtual FEditorModeTools& GetEditorModeManager() const override;
 	virtual UTypedElementCommonActions* GetCommonActions() const override;
 	virtual FName GetStatusBarName() const override;
@@ -365,7 +367,7 @@ private:
 	TSharedPtr<IDetailRootObjectCustomization> ActorDetailsRootCustomization;
 
 	/** Actor details SCS editor customization */
-	TSharedPtr<ISCSEditorUICustomization> ActorDetailsSCSEditorUICustomization;
+	TSharedPtr<FActorDetailsSCSEditorUICustomization> ActorDetailsSCSEditorUICustomization;
 		
 	/** A delegate which is called any time the LevelEditor's active viewport changes. */
 	FOnActiveViewportChanged OnActiveViewportChangedDelegate;
diff --git a/Engine/Source/Editor/LevelEditor/Public/ILevelEditor.h b/Engine/Source/Editor/LevelEditor/Public/ILevelEditor.h
--- a/Engine/Source/Editor/LevelEditor/Public/ILevelEditor.h
+++ b/Engine/Source/Editor/LevelEditor/Public/ILevelEditor.h
@@ -85,8 +85,14 @@ public:
 	/** Set the filter that should be used to determine the set of objects that should be shown in a details panel when an actor in the level editor is selected */
 	virtual void SetActorDetailsRootCustomization(TSharedPtr<FDetailsViewObjectFilter> ActorDetailsObjectFilter, TSharedPtr<IDetailRootObjectCustomization> ActorDetailsRootCustomization) = 0;
 
-	/** Sets the UI customization of the SCSEditor inside the level editor details panel. */
-	virtual void SetActorDetailsSCSEditorUICustomization(TSharedPtr<class ISCSEditorUICustomization> ActorDetailsSCSEditorUICustomization) = 0;
+	/** Adds the UI customization of the SCSEditor inside the level editor details panel. */
+	virtual void AddActorDetailsSCSEditorUICustomization(TSharedPtr<class ISCSEditorUICustomization> ActorDetailsSCSEditorUICustomization) = 0;
+	
+	/** Remove the UI customization of the SCSEditor inside the level editor details panel. */
+	virtual void RemoveActorDetailsSCSEditorUICustomization(TSharedPtr<class ISCSEditorUICustomization> ActorDetailsSCSEditorUICustomization) = 0;
+
+	UE_DEPRECATED(5.5, "Use AddActorDetailsSCSEditorUICustomization/RemoveActorDetailsSCSEditorUICustomization instead.")
+	virtual void SetActorDetailsSCSEditorUICustomization(TSharedPtr<class ISCSEditorUICustomization> ActorDetailsSCSEditorUICustomization) {};
 
 	/** Return the most recently interacted with Outliner */
 	UE_DEPRECATED(5.1, "The Level Editor has multiple outliners, use GetAllSceneOutliners() or GetMostRecentlyUsedSceneOutliner() instead to avoid ambiguity")
diff --git a/Engine/Source/Editor/LevelInstanceEditor/Private/LevelInstanceEditorModule.cpp b/Engine/Source/Editor/LevelInstanceEditor/Private/LevelInstanceEditorModule.cpp
--- a/Engine/Source/Editor/LevelInstanceEditor/Private/LevelInstanceEditorModule.cpp
+++ b/Engine/Source/Editor/LevelInstanceEditor/Private/LevelInstanceEditorModule.cpp
@@ -47,6 +47,7 @@
 #include "WorldPartition/WorldPartitionConverter.h"
 #include "WorldPartition/WorldPartitionActorLoaderInterface.h"
 #include "ScopedTransaction.h"
+#include "ISCSEditorUICustomization.h"
 
 IMPLEMENT_MODULE( FLevelInstanceEditorModule, LevelInstanceEditor );
 
@@ -903,8 +904,75 @@ struct FLevelInstanceMenuUtils
 	}
 };
 
+class FLevelInstanceActorDetailsSCSEditorUICustomization : public ISCSEditorUICustomization
+{
+public:
+	static TSharedPtr<FLevelInstanceActorDetailsSCSEditorUICustomization> GetInstance()
+	{
+		if (!Instance)
+		{
+			Instance = MakeShareable(new FLevelInstanceActorDetailsSCSEditorUICustomization());
+		}
+		return Instance;
+	}
+
+	virtual bool HideComponentsTree(TArrayView<UObject*> Context) const override { return false; }
+
+	virtual bool HideComponentsFilterBox(TArrayView<UObject*> Context) const override { return false; }
+
+	virtual bool HideAddComponentButton(TArrayView<UObject*> Context) const override { return ShouldHide(Context); }
+
+	virtual bool HideBlueprintButtons(TArrayView<UObject*> Context) const override { return ShouldHide(Context); }
+private:
+	bool ShouldHide(TArrayView<UObject*> Context) const
+	{
+		for (const UObject* ContextObject : Context)
+		{
+			if (const AActor* ActorContext = Cast<AActor>(ContextObject))
+			{
+				if (ActorContext->IsInLevelInstance() && !ActorContext->IsInEditLevelInstance())
+				{
+					return true;
+				}
+			}
+		}
+
+		return false;
+	}
+		
+	static TSharedPtr<FLevelInstanceActorDetailsSCSEditorUICustomization> Instance;
+	bool bShouldHide = false;
+};
+
+TSharedPtr<FLevelInstanceActorDetailsSCSEditorUICustomization> FLevelInstanceActorDetailsSCSEditorUICustomization::Instance;
+
+void FLevelInstanceEditorModule::OnLevelEditorCreated(TSharedPtr<ILevelEditor> InLevelEditor)
+{
+	RegisterToFirstLevelEditor();
+}
+
+void FLevelInstanceEditorModule::RegisterToFirstLevelEditor()
+{
+	FLevelEditorModule& LevelEditorModule = FModuleManager::GetModuleChecked<FLevelEditorModule>("LevelEditor");
+	TSharedPtr<ILevelEditor> FirstLevelEditor = LevelEditorModule.GetFirstLevelEditor();
+	if (FirstLevelEditor.IsValid())
+	{
+		FirstLevelEditor->AddActorDetailsSCSEditorUICustomization(FLevelInstanceActorDetailsSCSEditorUICustomization::GetInstance());
+	}
+}
+
 void FLevelInstanceEditorModule::StartupModule()
 {
+	FLevelEditorModule& LevelEditorModule = FModuleManager::LoadModuleChecked<FLevelEditorModule>("LevelEditor");
+	if (TSharedPtr<ILevelEditor> FirstLevelEditor = LevelEditorModule.GetFirstLevelEditor())
+	{
+		RegisterToFirstLevelEditor();
+	}
+	else
+	{
+		LevelEditorModule.OnLevelEditorCreated().AddRaw(this, &FLevelInstanceEditorModule::OnLevelEditorCreated);
+	}
+			
 	ExtendContextMenu();
 
 	FPropertyEditorModule& PropertyModule = FModuleManager::LoadModuleChecked<FPropertyEditorModule>("PropertyEditor");
@@ -965,6 +1033,16 @@ void FLevelInstanceEditorModule::StartupModule()
 
 void FLevelInstanceEditorModule::ShutdownModule()
 {
+	if (FModuleManager::Get().IsModuleLoaded("LevelEditor"))
+	{
+		FLevelEditorModule& LevelEditorModule = FModuleManager::GetModuleChecked<FLevelEditorModule>("LevelEditor");
+		LevelEditorModule.OnLevelEditorCreated().RemoveAll(this);
+		if (TSharedPtr<ILevelEditor> FirstLevelEditor = LevelEditorModule.GetFirstLevelEditor())
+		{
+			FirstLevelEditor->RemoveActorDetailsSCSEditorUICustomization(FLevelInstanceActorDetailsSCSEditorUICustomization::GetInstance());
+		}
+	}
+
 	if (GEditor)
 	{
 		GEditor->OnLevelActorDeleted().RemoveAll(this);
diff --git a/Engine/Source/Editor/LevelInstanceEditor/Private/LevelInstanceEditorModule.h b/Engine/Source/Editor/LevelInstanceEditor/Private/LevelInstanceEditorModule.h
--- a/Engine/Source/Editor/LevelInstanceEditor/Private/LevelInstanceEditorModule.h
+++ b/Engine/Source/Editor/LevelInstanceEditor/Private/LevelInstanceEditorModule.h
@@ -10,6 +10,7 @@
 class AActor;
 class ULevel;
 enum class EMapChangeType : uint8;
+class ILevelEditor;
 
 /**
  * The module holding all of the UI related pieces for LevelInstance management
@@ -48,6 +49,9 @@ private:
 
 	void ExtendContextMenu();
 		
+	void OnLevelEditorCreated(TSharedPtr<ILevelEditor> InLevelEditor);
+	void RegisterToFirstLevelEditor();
+
 	class FPropertyEditorPolicy : public PropertyEditorPolicy::IEditConstPolicy, 
 								  public PropertyEditorPolicy::IArchetypePolicy
 	{
diff --git a/Engine/Source/Editor/SubobjectEditor/Private/SSubobjectEditor.cpp b/Engine/Source/Editor/SubobjectEditor/Private/SSubobjectEditor.cpp
--- a/Engine/Source/Editor/SubobjectEditor/Private/SSubobjectEditor.cpp
+++ b/Engine/Source/Editor/SubobjectEditor/Private/SSubobjectEditor.cpp
@@ -2234,7 +2234,8 @@ TSharedRef<ITableRow> SSubobjectEditor::MakeTableRowWidget(FSubobjectEditorTreeN
 
 EVisibility SSubobjectEditor::GetComponentsTreeVisibility() const
 {
-	return (UICustomization.IsValid() && UICustomization->HideComponentsTree())
+	UObject* ObjectContextPtr = GetObjectContext();
+	return (UICustomization.IsValid() && UICustomization->HideComponentsTree(MakeArrayView(&ObjectContextPtr, 1)))
 		       ? EVisibility::Collapsed
 		       : EVisibility::Visible;
 }
@@ -2393,7 +2394,8 @@ void SSubobjectEditor::SelectNode(FSubobjectEditorTreeNodePtrType InNodeToSelect
 
 EVisibility SSubobjectEditor::GetComponentsFilterBoxVisibility() const
 {
-	return (UICustomization.IsValid() && UICustomization->HideComponentsFilterBox())
+	UObject* ObjectContextPtr = GetObjectContext();
+	return (UICustomization.IsValid() && UICustomization->HideComponentsFilterBox(MakeArrayView(&ObjectContextPtr, 1)))
 		       ? EVisibility::Collapsed
 		       : EVisibility::Visible;
 }
@@ -3296,8 +3298,9 @@ void SSubobjectEditor::GetSelectedItemsForContextMenu(TArray<FComponentEventCons
 
 TSubclassOf<UActorComponent> SSubobjectEditor::GetComponentTypeFilterToApply() const
 {
+	UObject* ObjectContextPtr = GetObjectContext();
 	TSubclassOf<UActorComponent> ComponentType = UICustomization.IsValid()
-		                                             ? UICustomization->GetComponentTypeFilter()
+		                                             ? UICustomization->GetComponentTypeFilter(MakeArrayView(&ObjectContextPtr, 1))
 		                                             : nullptr;
 	if (!ComponentType)
 	{
@@ -3585,7 +3588,8 @@ FSubobjectDataHandle SSubobjectEditor::PerformComboAddClass(TSubclassOf<UActorCo
 
 EVisibility SSubobjectEditor::GetPromoteToBlueprintButtonVisibility() const
 {
-	return (UICustomization.IsValid() && UICustomization->HideBlueprintButtons())
+	UObject* ObjectContextPtr = GetObjectContext();
+	return (UICustomization.IsValid() && UICustomization->HideBlueprintButtons(MakeArrayView(&ObjectContextPtr, 1)))
 	       || (ShowInlineSearchWithButtons())
 	       || (GetBlueprint() != nullptr)
 		       ? EVisibility::Collapsed
@@ -3594,7 +3598,8 @@ EVisibility SSubobjectEditor::GetPromoteToBlueprintButtonVisibility() const
 
 EVisibility SSubobjectEditor::GetEditBlueprintButtonVisibility() const
 {
-	return (UICustomization.IsValid() && UICustomization->HideBlueprintButtons())
+	UObject* ObjectContextPtr = GetObjectContext();
+	return (UICustomization.IsValid() && UICustomization->HideBlueprintButtons(MakeArrayView(&ObjectContextPtr, 1)))
 	       || (ShowInlineSearchWithButtons())
 	       || (GetBlueprint() == nullptr)
 		       ? EVisibility::Collapsed
@@ -3603,8 +3608,9 @@ EVisibility SSubobjectEditor::GetEditBlueprintButtonVisibility() const
 
 EVisibility SSubobjectEditor::GetComponentClassComboButtonVisibility() const
 {
+	UObject* ObjectContextPtr = GetObjectContext();
 	return (HideComponentClassCombo.Get() 
-	|| (UICustomization.IsValid() && UICustomization->HideAddComponentButton())) 
+	|| (UICustomization.IsValid() && UICustomization->HideAddComponentButton(MakeArrayView(&ObjectContextPtr, 1))))
 	? EVisibility::Collapsed : EVisibility::Visible;
 }
 
