Level Instance: Sub Actor Selection support
- SceneOutliner: allow to select Level Instance sub actors in the outliner in a read only way (allows copy/pasting of actors and properties)
- Viewport: allow sub selection through "Shift+ Scroll Up/Down"
- Gated by new settings in ULevelInstanceEditorPerProjectUserSettings

#rnx

diff --git a/Engine/Source/Editor/LevelInstanceEditor/LevelInstanceEditor.Build.cs b/Engine/Source/Editor/LevelInstanceEditor/LevelInstanceEditor.Build.cs
--- a/Engine/Source/Editor/LevelInstanceEditor/LevelInstanceEditor.Build.cs
+++ b/Engine/Source/Editor/LevelInstanceEditor/LevelInstanceEditor.Build.cs
@@ -35,7 +35,11 @@ public class LevelInstanceEditor : ModuleRules
 				"SceneOutliner",
 				"WorldPartitionEditor",
 				"Kismet",
-				"PropertyEditor"
+				"PropertyEditor",
+				"InteractiveToolsFramework",
+				"EditorInteractiveToolsFramework",
+				"TypedElementFramework",
+				"TypedElementRuntime"
 			}
 		);
 
diff --git a/Engine/Source/Editor/LevelInstanceEditor/Private/LevelInstanceEditorMode.cpp b/Engine/Source/Editor/LevelInstanceEditor/Private/LevelInstanceEditorMode.cpp
--- a/Engine/Source/Editor/LevelInstanceEditor/Private/LevelInstanceEditorMode.cpp
+++ b/Engine/Source/Editor/LevelInstanceEditor/Private/LevelInstanceEditorMode.cpp
@@ -3,10 +3,12 @@
 #include "LevelInstanceEditorMode.h"
 #include "LevelInstanceEditorModeToolkit.h"
 #include "LevelInstanceEditorModeCommands.h"
+#include "LevelInstanceEditorSettings.h"
 #include "Editor.h"
 #include "Selection.h"
 #include "EditorModes.h"
 #include "Engine/World.h"
+#include "EngineUtils.h"
 #include "LevelInstance/LevelInstanceSubsystem.h"
 #include "LevelInstance/LevelInstanceInterface.h"
 #include "LevelInstance/ILevelInstanceEditorModule.h"
@@ -15,11 +17,143 @@
 #include "EditorModeManager.h"
 #include "Framework/Application/SlateApplication.h"
 #include "Modules/ModuleManager.h"
+#include "InteractiveToolManager.h"
+#include "EdModeInteractiveToolsContext.h"
+#include "BaseBehaviors/MouseWheelBehavior.h"
+#include "InputRouter.h"
+#include "ToolContextInterfaces.h"
+#include "Elements/Framework/EngineElementsLibrary.h"
+#include "Elements/Framework/TypedElementHandle.h"
 
 #define LOCTEXT_NAMESPACE "LevelInstanceEditorMode"
 
 FEditorModeID ULevelInstanceEditorMode::EM_LevelInstanceEditorModeId("EditMode.LevelInstance");
 
+class FMouseWheelBehaviorTarget : public IMouseWheelBehaviorTarget
+{
+public:
+	FMouseWheelBehaviorTarget(UEditorInteractiveToolsContext* InInteractiveToolContext) : InteractiveToolContext(InInteractiveToolContext) {}
+
+	UEditorInteractiveToolsContext* InteractiveToolContext = nullptr;
+
+	// IMouseWheelBehaviorTarget
+	virtual FInputRayHit ShouldRespondToMouseWheel(const FInputDeviceRay& CurrentPos) override
+	{
+		FInputRayHit ToReturn;
+
+		TArray<AActor*> SelectionHierarchy;
+		int32 SelectionIndex = INDEX_NONE;
+		ToReturn.bHit = GetLevelInstanceSelectionHierarchy(CurrentPos, SelectionHierarchy, SelectionIndex);
+
+		return ToReturn;
+	}
+	virtual void OnMouseWheelScrollUp(const FInputDeviceRay& CurrentPos) override
+	{
+		TArray<AActor*> SelectionHierarchy;
+		int32 SelectionIndex = INDEX_NONE;
+		if (GetLevelInstanceSelectionHierarchy(CurrentPos, SelectionHierarchy, SelectionIndex))
+		{
+			SelectActorAt(SelectionIndex - 1, SelectionHierarchy);
+		}
+	}
+
+	virtual void OnMouseWheelScrollDown(const FInputDeviceRay& CurrentPos) override
+	{
+		TArray<AActor*> SelectionHierarchy;
+		int32 SelectionIndex = INDEX_NONE;
+		if (GetLevelInstanceSelectionHierarchy(CurrentPos, SelectionHierarchy, SelectionIndex))
+		{
+			SelectActorAt(SelectionIndex + 1, SelectionHierarchy);
+		}
+	}
+
+private:
+	bool GetLevelInstanceSelectionHierarchy(const FInputDeviceRay& CurrentPos, TArray<AActor*>& OutSelectionHierarchy, int32& OutSelectionIndex) const
+	{
+		if (!GetDefault<ULevelInstanceEditorPerProjectUserSettings>()->bIsViewportSubSelectionEnabled)
+		{
+			return false;
+		}
+
+		if (!FSlateApplication::Get().GetModifierKeys().IsShiftDown())
+		{
+			return false;
+		}
+
+		TArray<AActor*> SelectedActors;
+		GEditor->GetSelectedActors()->GetSelectedObjects<AActor>(SelectedActors);
+
+		// Only handle mouse wheel on single selection
+		if (SelectedActors.Num() != 1)
+		{
+			return false;
+		}
+
+		AActor* SelectedActor = SelectedActors[0];
+
+		if (CurrentPos.bHas2D)
+		{
+			if (IToolsContextQueriesAPI* ContextAPI = InteractiveToolContext->ToolManager->GetContextQueriesAPI())
+			{
+				if (FViewport* Viewport = ContextAPI->GetFocusedViewport())
+				{
+					if (HHitProxy* HitResult = Viewport->GetHitProxy(CurrentPos.ScreenPosition.X, CurrentPos.ScreenPosition.Y))
+					{
+						if (HActor* HitActor = HitProxyCast<HActor>(HitResult))
+						{
+							if (AActor* Actor = HitActor->Actor; Actor && Actor->IsInLevelInstance())
+							{
+								OutSelectionIndex = Actor == SelectedActor ? 0 : INDEX_NONE;
+								OutSelectionHierarchy.Add(Actor);
+
+								ULevelInstanceSubsystem* LevelInstanceSubsystem = UWorld::GetSubsystem<ULevelInstanceSubsystem>(Actor->GetWorld());
+								LevelInstanceSubsystem->ForEachLevelInstanceAncestors(Actor, [SelectedActor, &OutSelectionIndex, &OutSelectionHierarchy](ILevelInstanceInterface* LevelInstanceInterface)
+									{
+										if (AActor* LevelInstanceActor = Cast<AActor>(LevelInstanceInterface))
+										{
+											OutSelectionIndex = LevelInstanceActor == SelectedActor ? OutSelectionHierarchy.Num() : OutSelectionIndex;
+											OutSelectionHierarchy.Add(LevelInstanceActor);
+										}
+										return true;
+									});
+
+								return OutSelectionIndex != INDEX_NONE && OutSelectionHierarchy.Num() > 1;
+							}
+						}
+					}
+				}
+			}
+		}
+
+		return false;
+	}
+
+	void SelectActorAt(int32 SelectionIndex, const TArray<AActor*>& SelectionHierarchy)
+	{
+		if (SelectionHierarchy.IsValidIndex(SelectionIndex))
+		{
+			AActor* ActorToSelect = SelectionHierarchy[SelectionIndex];
+			if (ActorToSelect->SupportsSubRootSelection())
+			{
+				if (UTypedElementSelectionSet* SelectionSet = GEditor->GetSelectedActors()->GetElementSelectionSet())
+				{
+					const FTypedElementSelectionOptions SelectionOptions = FTypedElementSelectionOptions()
+						.SetAllowHidden(true)
+						.SetWarnIfLocked(false)
+						.SetAllowLegacyNotifications(false)
+						.SetAllowSubRootSelection(true);
+
+					FTypedElementHandle ActorElementHandle = UEngineElementsLibrary::AcquireEditorActorElementHandle(SelectionHierarchy[SelectionIndex]);
+					if(SelectionSet->CanSelectElement(ActorElementHandle, SelectionOptions))
+					{
+						SelectionSet->SetSelection(MakeArrayView(&ActorElementHandle,1), SelectionOptions);
+					}
+				}
+			}
+		}
+	}
+};
+
 ULevelInstanceEditorMode::ULevelInstanceEditorMode()
 	: UEdMode()
 {
@@ -35,6 +169,22 @@ ULevelInstanceEditorMode::~ULevelInstanceEditorMode()
 {
 }
 
+void ULevelInstanceEditorBehaviorSource::Initialize(UEditorInteractiveToolsContext* InteractiveToolsContext)
+{
+	InputBehaviorSet = NewObject<UInputBehaviorSet>();
+	UMouseWheelInputBehavior* MouseWheelInputBehavior = NewObject<UMouseWheelInputBehavior>();
+	MouseWheelBehaviorTarget = MakeUnique<FMouseWheelBehaviorTarget>(InteractiveToolsContext);
+	MouseWheelInputBehavior->Initialize(MouseWheelBehaviorTarget.Get());
+	InputBehaviorSet->Add(MouseWheelInputBehavior);
+}
+
+TScriptInterface<IInputBehaviorSource> ULevelInstanceEditorMode::CreateDefaultModeBehaviorSource(UEditorInteractiveToolsContext* InteractiveToolContext)
+{
+	ULevelInstanceEditorBehaviorSource* NewBehaviorSource = NewObject<ULevelInstanceEditorBehaviorSource>();
+	NewBehaviorSource->Initialize(InteractiveToolContext);
+	return NewBehaviorSource;
+}
+
 void ULevelInstanceEditorMode::OnPreBeginPIE(bool bSimulate)
 {
 	ExitModeCommand();
@@ -62,12 +212,24 @@ void ULevelInstanceEditorMode::Enter()
 	UEdMode::Enter();
 
 	UpdateEngineShowFlags();
+	
+	if (UEditorInteractiveToolsContext* InteractiveToolContext = GetInteractiveToolsContext(EToolsContextScope::EdMode))
+	{
+		// Here we create a BehaviorSource specific to the Level Instance Editor Mode, for now it is the same type as the default one.
+		ModeBehaviorSource = CreateDefaultModeBehaviorSource(InteractiveToolContext);
+		InteractiveToolContext->InputRouter->RegisterSource(ModeBehaviorSource.GetInterface());
+	}
 
 	FEditorDelegates::PreBeginPIE.AddUObject(this, &ULevelInstanceEditorMode::OnPreBeginPIE);
 }
 
 void ULevelInstanceEditorMode::Exit()
 {
+	if (UEditorInteractiveToolsContext* InteractiveToolContext = GetInteractiveToolsContext(EToolsContextScope::EdMode))
+	{
+		InteractiveToolContext->InputRouter->DeregisterSource(ModeBehaviorSource.GetInterface());
+	}
+
 	UEdMode::Exit();
 		
 	UpdateEngineShowFlags();
diff --git a/Engine/Source/Editor/LevelInstanceEditor/Private/LevelInstanceEditorMode.h b/Engine/Source/Editor/LevelInstanceEditor/Private/LevelInstanceEditorMode.h
--- a/Engine/Source/Editor/LevelInstanceEditor/Private/LevelInstanceEditorMode.h
+++ b/Engine/Source/Editor/LevelInstanceEditor/Private/LevelInstanceEditorMode.h
@@ -4,8 +4,25 @@
 #include "CoreMinimal.h"
 #include "Tools/UEdMode.h"
 #include "Framework/Commands/UICommandList.h"
+#include "InputBehaviorSet.h"
+#include "BaseBehaviors/BehaviorTargetInterfaces.h"
 #include "LevelInstanceEditorMode.generated.h"
 
+UCLASS()
+class ULevelInstanceEditorBehaviorSource : public UObject, public IInputBehaviorSource
+{
+	GENERATED_BODY()
+public:
+	virtual const UInputBehaviorSet* GetInputBehaviors() const override { return InputBehaviorSet; }
+	void Initialize(UEditorInteractiveToolsContext* InteractiveToolsContext);
+
+private:
+	UPROPERTY(Transient)
+	TObjectPtr<UInputBehaviorSet> InputBehaviorSet;
+
+	TUniquePtr<IMouseWheelBehaviorTarget> MouseWheelBehaviorTarget;
+};
+
 UCLASS()
 class ULevelInstanceEditorMode : public UEdMode
 {
@@ -33,6 +50,8 @@ public:
 	virtual bool IsOperationSupportedForCurrentAsset(EAssetOperation InOperation) const { return InOperation == EAssetOperation::Save; }
 	// End UEdMode
 
+	static TScriptInterface<IInputBehaviorSource> CreateDefaultModeBehaviorSource(UEditorInteractiveToolsContext* InteractiveToolContext);
+	
 private:
 	void OnPreBeginPIE(bool bSimulate);
 	void UpdateEngineShowFlags();
@@ -44,5 +63,8 @@ private:
 	bool IsContextRestrictedForWorld(UWorld* InWorld) const;
 
 	bool bContextRestriction;
+	
+	UPROPERTY(Transient)
+	TScriptInterface<IInputBehaviorSource> ModeBehaviorSource;
 };
 
diff --git a/Engine/Source/Editor/LevelInstanceEditor/Private/LevelInstanceEditorModule.cpp b/Engine/Source/Editor/LevelInstanceEditor/Private/LevelInstanceEditorModule.cpp
--- a/Engine/Source/Editor/LevelInstanceEditor/Private/LevelInstanceEditorModule.cpp
+++ b/Engine/Source/Editor/LevelInstanceEditor/Private/LevelInstanceEditorModule.cpp
@@ -48,6 +48,7 @@
 #include "WorldPartition/WorldPartitionActorLoaderInterface.h"
 #include "ScopedTransaction.h"
 #include "ISCSEditorUICustomization.h"
+#include "EdModeInteractiveToolsContext.h"
 
 IMPLEMENT_MODULE( FLevelInstanceEditorModule, LevelInstanceEditor );
 
@@ -1028,6 +1029,10 @@ void FLevelInstanceEditorModule::StartupModule()
 	if (!IsRunningCommandlet())
 	{
 		GLevelEditorModeTools().OnEditorModeIDChanged().AddRaw(this, &FLevelInstanceEditorModule::OnEditorModeIDChanged);
+	
+		// Create a Behavior source for the default EdModeTools (when we aren't in the LevelInstanceEditorMode)
+		DefaultBehaviorSource = ULevelInstanceEditorMode::CreateDefaultModeBehaviorSource(GLevelEditorModeTools().GetInteractiveToolsContext());
+		GLevelEditorModeTools().GetInteractiveToolsContext()->InputRouter->RegisterSource(DefaultBehaviorSource.GetInterface());
 	}
 
 	ULevelInstanceSubsystem::RegisterPrimitiveColorHandler();
@@ -1057,6 +1062,8 @@ void FLevelInstanceEditorModule::ShutdownModule()
 	if (!IsRunningCommandlet() && GLevelEditorModeToolsIsValid())
 	{
 		GLevelEditorModeTools().OnEditorModeIDChanged().RemoveAll(this);
+		GLevelEditorModeTools().GetInteractiveToolsContext()->InputRouter->DeregisterSource(DefaultBehaviorSource.GetInterface());
+		DefaultBehaviorSource = nullptr;
 	}
 }
 
@@ -1254,4 +1261,14 @@ bool FLevelInstanceEditorModule::IsEditInPlaceStreamingEnabled() const
 	return GetDefault<ULevelInstanceEditorSettings>()->bIsEditInPlaceStreamingEnabled;
 }
 
+bool FLevelInstanceEditorModule::IsSubSelectionEnabled() const
+{
+	return GetDefault<ULevelInstanceEditorPerProjectUserSettings>()->bIsSubSelectionEnabled;
+}
+
+void FLevelInstanceEditorModule::AddReferencedObjects(FReferenceCollector& Collector)
+{
+	DefaultBehaviorSource.AddReferencedObjects(Collector);
+}
+
 #undef LOCTEXT_NAMESPACE
\ No newline at end of file
diff --git a/Engine/Source/Editor/LevelInstanceEditor/Private/LevelInstanceEditorModule.h b/Engine/Source/Editor/LevelInstanceEditor/Private/LevelInstanceEditorModule.h
--- a/Engine/Source/Editor/LevelInstanceEditor/Private/LevelInstanceEditorModule.h
+++ b/Engine/Source/Editor/LevelInstanceEditor/Private/LevelInstanceEditorModule.h
@@ -3,19 +3,21 @@
 
 #include "CoreMinimal.h"
 #include "LevelInstance/ILevelInstanceEditorModule.h"
+#include "UObject/GCObject.h"
 #include "PropertyEditorArchetypePolicy.h"
 #include "PropertyEditorEditConstPolicy.h"
 #include "Tools/Modes.h"
 
 class AActor;
 class ULevel;
+class IInputBehaviorSource;
 enum class EMapChangeType : uint8;
 class ILevelEditor;
 
 /**
  * The module holding all of the UI related pieces for LevelInstance management
  */
-class FLevelInstanceEditorModule : public ILevelInstanceEditorModule
+class FLevelInstanceEditorModule : public ILevelInstanceEditorModule, public FGCObject
 {
 public:
 	virtual ~FLevelInstanceEditorModule(){}
@@ -39,7 +41,10 @@ public:
 	virtual FTryExitEditorModeEvent& OnTryExitEditorMode() override { return TryExitEditorModeEvent; }
 
 	virtual bool IsEditInPlaceStreamingEnabled() const override;
+	virtual bool IsSubSelectionEnabled() const override;
 
+	virtual void AddReferencedObjects(FReferenceCollector& Collector) override;
+	virtual FString GetReferencerName() const override { return TEXT("LevelInstanceEditorModule"); }
 private:
 	virtual void UpdateEditorMode(bool bActivated) override;
 	
@@ -116,4 +121,6 @@ private:
 	
 	FExitEditorModeEvent ExitEditorModeEvent;
 	FTryExitEditorModeEvent TryExitEditorModeEvent;
+
+	TScriptInterface<IInputBehaviorSource> DefaultBehaviorSource;
 };
diff --git a/Engine/Source/Editor/LevelInstanceEditor/Public/LevelInstanceEditorSettings.h b/Engine/Source/Editor/LevelInstanceEditor/Public/LevelInstanceEditorSettings.h
--- a/Engine/Source/Editor/LevelInstanceEditor/Public/LevelInstanceEditorSettings.h
+++ b/Engine/Source/Editor/LevelInstanceEditor/Public/LevelInstanceEditorSettings.h
@@ -52,6 +52,14 @@ public:
 	UPROPERTY(config, EditAnywhere, Category = Pivot)
 	ELevelInstancePivotType PivotType;
 
+	/** If true, allows read only selection of actors inside a Level Instance */
+	UPROPERTY(config, EditAnywhere, Category = "Selection", meta = (ToolTip = "Allow Read Only Selection of Level Instance Level Actors"))
+	bool bIsSubSelectionEnabled;
+
+	/** If true, allows read only selection of actors inside a Level Instance through the viewport */
+	UPROPERTY(config, EditAnywhere, Category = "Selection", meta = (DisplayName = "Is Viewport Sub Selection Enabled (Shift + Mouse Scroll Up/Down)", ToolTip = "Allow Read Only Selection of Level Instance Level Actors in the Viewport (Shortcut: Shift + Mouse Scroll Up/Down)"))
+	bool bIsViewportSubSelectionEnabled;
+
 	/**
 	 * When the Level Instance is broken via "Level->Break..", its actors will be placed inside the folder the LI is
 	 * inside of, under a subfolder with the name of the Level Instance, and also keeping their original folder structure.
diff --git a/Engine/Source/Runtime/Engine/Classes/GameFramework/Actor.h b/Engine/Source/Runtime/Engine/Classes/GameFramework/Actor.h
--- a/Engine/Source/Runtime/Engine/Classes/GameFramework/Actor.h
+++ b/Engine/Source/Runtime/Engine/Classes/GameFramework/Actor.h
@@ -1076,7 +1076,7 @@ public:
 	inline const FGuid& GetContentBundleGuid() const { return ContentBundleGuid; }
 
 	/** Returns true if actor location should be locked. */
-	virtual bool IsLockLocation() const { return bLockLocation; }
+	ENGINE_API virtual bool IsLockLocation() const;
 
 	/** Set the bLockLocation flag */
 	void SetLockLocation(bool bInLockLocation) { bLockLocation = bInLockLocation; }
@@ -3063,7 +3063,7 @@ public:
 	ENGINE_API virtual AActor* GetRootSelectionParent() const;
 
 	/** Returns true if actor can be selected as a sub selection of its root selection parent */
-	ENGINE_API virtual bool SupportsSubRootSelection() const { return false; }
+	ENGINE_API virtual bool SupportsSubRootSelection() const;
 
 	/** Returns if actor or selection parent is selected */
 	ENGINE_API bool IsActorOrSelectionParentSelected() const;
diff --git a/Engine/Source/Runtime/Engine/Private/Actor.cpp b/Engine/Source/Runtime/Engine/Private/Actor.cpp
--- a/Engine/Source/Runtime/Engine/Private/Actor.cpp
+++ b/Engine/Source/Runtime/Engine/Private/Actor.cpp
@@ -5206,6 +5206,18 @@ AActor* AActor::GetSelectionParent() const
 	return nullptr;
 }
 
+bool AActor::SupportsSubRootSelection() const
+{
+#if WITH_EDITOR
+	if (IsInLevelInstance())
+	{
+		return UWorld::GetSubsystem<ULevelInstanceSubsystem>(GetWorld())->IsSubSelectionEnabled();
+	}
+#endif
+
+	return false;
+}
+
 AActor* AActor::GetRootSelectionParent() const
 {
 	AActor* Parent = GetSelectionParent();
diff --git a/Engine/Source/Runtime/Engine/Private/ActorEditor.cpp b/Engine/Source/Runtime/Engine/Private/ActorEditor.cpp
--- a/Engine/Source/Runtime/Engine/Private/ActorEditor.cpp
+++ b/Engine/Source/Runtime/Engine/Private/ActorEditor.cpp
@@ -1140,6 +1140,16 @@ void AActor::SetPackageExternal(bool bExternal, bool bShouldDirty, UPackage* Act
 	MarkPackageDirty();
 }
 
+bool AActor::IsLockLocation() const
+{
+	if (IsInLevelInstance() && !IsInAnyEditLevelInstance())
+	{
+		return true;
+	}
+
+	return bLockLocation;
+}
+
 void AActor::OnPlayFromHere()
 {
 	check(bCanPlayFromHere);
diff --git a/Engine/Source/Runtime/Engine/Private/LevelInstance/LevelInstanceSubsystem.cpp b/Engine/Source/Runtime/Engine/Private/LevelInstance/LevelInstanceSubsystem.cpp
--- a/Engine/Source/Runtime/Engine/Private/LevelInstance/LevelInstanceSubsystem.cpp
+++ b/Engine/Source/Runtime/Engine/Private/LevelInstance/LevelInstanceSubsystem.cpp
@@ -2294,6 +2294,12 @@ bool ULevelInstanceSubsystem::CanEditLevelInstance(const ILevelInstanceInterface
 	return true;
 }
 
+bool ULevelInstanceSubsystem::IsSubSelectionEnabled() const
+{
+	ILevelInstanceEditorModule& EditorModule = FModuleManager::GetModuleChecked<ILevelInstanceEditorModule>("LevelInstanceEditor");
+	return EditorModule.IsSubSelectionEnabled();
+}
+
 bool ULevelInstanceSubsystem::CanEditLevelInstancePropertyOverrides(const ILevelInstanceInterface* LevelInstance, FText* OutReason) const
 {
 	if (!CanEditLevelInstanceCommon(LevelInstance, OutReason))
diff --git a/Engine/Source/Runtime/Engine/Public/LevelInstance/ILevelInstanceEditorModule.h b/Engine/Source/Runtime/Engine/Public/LevelInstance/ILevelInstanceEditorModule.h
--- a/Engine/Source/Runtime/Engine/Public/LevelInstance/ILevelInstanceEditorModule.h
+++ b/Engine/Source/Runtime/Engine/Public/LevelInstance/ILevelInstanceEditorModule.h
@@ -30,7 +30,7 @@ public:
 	virtual FTryExitEditorModeEvent& OnTryExitEditorMode() = 0;
 
 	virtual bool IsEditInPlaceStreamingEnabled() const = 0;
-
+	virtual bool IsSubSelectionEnabled() const = 0;
 protected:
 	friend class ULevelInstanceSubsystem;
 	friend class ULevelStreamingLevelInstanceEditorPropertyOverride;
diff --git a/Engine/Source/Runtime/Engine/Public/LevelInstance/LevelInstanceSubsystem.h b/Engine/Source/Runtime/Engine/Public/LevelInstance/LevelInstanceSubsystem.h
--- a/Engine/Source/Runtime/Engine/Public/LevelInstance/LevelInstanceSubsystem.h
+++ b/Engine/Source/Runtime/Engine/Public/LevelInstance/LevelInstanceSubsystem.h
@@ -81,6 +81,7 @@ public:
 	ENGINE_API bool IsLoaded(const ILevelInstanceInterface* LevelInstance) const;
 	ENGINE_API bool IsLoading(const ILevelInstanceInterface* LevelInstance) const;
 	ENGINE_API void ForEachLevelInstanceAncestors(const AActor* Actor, TFunctionRef<bool(const ILevelInstanceInterface*)> Operation) const;
+	ENGINE_API void ForEachLevelInstanceAncestors(AActor* Actor, TFunctionRef<bool(ILevelInstanceInterface*)> Operation) const;
 	ENGINE_API void ForEachLevelInstanceAncestorsAndSelf(AActor* Actor, TFunctionRef<bool(ILevelInstanceInterface*)> Operation) const;
 	/** Runs a lambda operation along the ancestors that own the LevelInstance. Primarily for capturing inclusive true/false by using lambda captures */
 	ENGINE_API void ForEachLevelInstanceAncestorsAndSelf(const AActor* Actor, TFunctionRef<bool(const ILevelInstanceInterface*)> Operation) const;
@@ -170,6 +171,8 @@ public:
 	ENGINE_API bool PassLevelInstanceFilter(UWorld* World, const FWorldPartitionHandle& Actor) const;
 
 	ENGINE_API bool IsEditingLevelInstancePropertyOverrides(const ILevelInstanceInterface* LevelInstance) const;
+
+	ENGINE_API bool IsSubSelectionEnabled() const;
 private:
 	friend class FLevelInstanceEditorModule;
 	friend struct FLevelInstanceMenuUtils;
@@ -199,7 +202,6 @@ private:
 	ENGINE_API void LoadLevelInstance(ILevelInstanceInterface* LevelInstance);
 	ENGINE_API void UnloadLevelInstance(const FLevelInstanceID& LevelInstanceID);
 	ENGINE_API void ForEachActorInLevel(ULevel* Level, TFunctionRef<bool(AActor * LevelActor)> Operation) const;
-	ENGINE_API void ForEachLevelInstanceAncestors(AActor* Actor, TFunctionRef<bool(ILevelInstanceInterface*)> Operation) const;
 	ENGINE_API void RegisterLoadedLevelStreamingLevelInstance(ULevelStreamingLevelInstance* LevelStreaming);
 
 #if WITH_EDITOR
