Level Instance: Scene Outliner Property Override Column
- Make use of HasEditableLevelInstancePropertyOverrides/HasLevelInstancePropertyOverrides in a new column to give user visual feedback on what is overriden
- Fix small Reset Property Override menu bug (not allowing resetting actor overrides on a level instance)

#rnx

diff --git a/Engine/Source/Editor/EditorStyle/Private/StarshipStyle.cpp b/Engine/Source/Editor/EditorStyle/Private/StarshipStyle.cpp
--- a/Engine/Source/Editor/EditorStyle/Private/StarshipStyle.cpp
+++ b/Engine/Source/Editor/EditorStyle/Private/StarshipStyle.cpp
@@ -2044,6 +2044,19 @@ void FStarshipEditorStyle::FStyle::SetupWorldPartitionStyles()
 	Set("WorldPartition.FocusSelection", new IMAGE_BRUSH_SVG("WorldPartition/FramingSelection_20", Icon16x16));
 	Set("WorldPartition.FocusLoadedRegions", new IMAGE_BRUSH_SVG("WorldPartition/FrameRegions_20", Icon16x16));
 	Set("WorldPartition.FocusWorld", new IMAGE_BRUSH_SVG("WorldPartition/FrameWorld_20", Icon16x16));
+
+	// Level Instance
+	Set("LevelInstance.ColumnOverrideHereEditable", new IMAGE_BRUSH_SVG("Starship/LevelInstance/ColumnOverrideHere", Icon16x16, FStyleColors::AccentBlue));
+	Set("LevelInstance.ColumnOverrideHere", new IMAGE_BRUSH_SVG("Starship/LevelInstance/ColumnOverrideHere", Icon16x16, FStyleColors::AccentGray));
+
+	Set("LevelInstance.ColumnOverrideContainerHereEditable", new IMAGE_BRUSH_SVG("Starship/LevelInstance/ColumnOverrideContainerHere", Icon16x16, FStyleColors::AccentBlue));
+	Set("LevelInstance.ColumnOverrideContainerHere", new IMAGE_BRUSH_SVG("Starship/LevelInstance/ColumnOverrideContainerHere", Icon16x16, FStyleColors::AccentGray));
+
+	Set("LevelInstance.ColumnOverrideContainerInsideEditable", new IMAGE_BRUSH_SVG("Starship/LevelInstance/ColumnOverrideContainerInside", Icon16x16, FStyleColors::AccentBlue));
+	Set("LevelInstance.ColumnOverrideContainerInside", new IMAGE_BRUSH_SVG("Starship/LevelInstance/ColumnOverrideContainerInside", Icon16x16, FStyleColors::AccentGray));
+
+	Set("LevelInstance.ColumnOverrideContainerEditable", new IMAGE_BRUSH_SVG("Starship/LevelInstance/ColumnOverrideContainer", Icon16x16, FStyleColors::AccentBlue));
+	Set("LevelInstance.ColumnOverrideContainer", new IMAGE_BRUSH_SVG("Starship/LevelInstance/ColumnOverrideContainer", Icon16x16, FStyleColors::AccentGray));
 }
 
 void FStarshipEditorStyle::FStyle::SetupSequencerStyles()
diff --git a/Engine/Source/Editor/LevelInstanceEditor/LevelInstanceEditor.Build.cs b/Engine/Source/Editor/LevelInstanceEditor/LevelInstanceEditor.Build.cs
--- a/Engine/Source/Editor/LevelInstanceEditor/LevelInstanceEditor.Build.cs
+++ b/Engine/Source/Editor/LevelInstanceEditor/LevelInstanceEditor.Build.cs
@@ -39,7 +39,8 @@ public class LevelInstanceEditor : ModuleRules
 				"InteractiveToolsFramework",
 				"EditorInteractiveToolsFramework",
 				"TypedElementFramework",
-				"TypedElementRuntime"
+				"TypedElementRuntime",
+				"SceneOutliner"
 			}
 		);
 
diff --git a/Engine/Source/Editor/LevelInstanceEditor/Private/LevelInstanceEditorModule.cpp b/Engine/Source/Editor/LevelInstanceEditor/Private/LevelInstanceEditorModule.cpp
--- a/Engine/Source/Editor/LevelInstanceEditor/Private/LevelInstanceEditorModule.cpp
+++ b/Engine/Source/Editor/LevelInstanceEditor/Private/LevelInstanceEditorModule.cpp
@@ -2,6 +2,7 @@
 #include "LevelInstanceEditorModule.h"
 #include "LevelInstanceActorDetails.h"
 #include "LevelInstancePivotDetails.h"
+#include "LevelInstanceSceneOutlinerColumn.h"
 #include "PackedLevelActorUtils.h"
 #include "LevelInstanceFilterPropertyTypeCustomization.h"
 #include "LevelInstance/LevelInstanceSubsystem.h"
@@ -49,6 +50,8 @@
 #include "ScopedTransaction.h"
 #include "ISCSEditorUICustomization.h"
 #include "EdModeInteractiveToolsContext.h"
+#include "SceneOutlinerModule.h"
+#include "SceneOutlinerFwd.h"
 
 IMPLEMENT_MODULE( FLevelInstanceEditorModule, LevelInstanceEditor );
 
@@ -647,28 +650,32 @@ struct FLevelInstanceMenuUtils
 				{
 					if (!LevelInstanceSubsystem->CanResetPropertyOverrides(SelectedLevelInstance))
 					{
-						return;
+						bCanResetAllLevelInstances = false;
+						break;
 					}
 				}
 
-				FToolMenuSection& Section = CreateLevelSection(Menu);
-				FToolUIAction UIAction;
-				UIAction.ExecuteAction.BindLambda([LevelInstanceSubsystem, CopySelectedLevelInstance = SelectedLevelInstances](const FToolMenuContext& MenuContext)
+				if (bCanResetAllLevelInstances)
 				{
-					for (ILevelInstanceInterface* LevelInstanceInterface : CopySelectedLevelInstance)
-					{
-						LevelInstanceSubsystem->ResetPropertyOverrides(LevelInstanceInterface);
-					}
-				});
+					FToolMenuSection& Section = CreateLevelSection(Menu);
+					FToolUIAction UIAction;
+					UIAction.ExecuteAction.BindLambda([LevelInstanceSubsystem, CopySelectedLevelInstance = SelectedLevelInstances](const FToolMenuContext& MenuContext)
+						{
+							for (ILevelInstanceInterface* LevelInstanceInterface : CopySelectedLevelInstance)
+							{
+								LevelInstanceSubsystem->ResetPropertyOverrides(LevelInstanceInterface);
+							}
+						});
 
-				Section.AddMenuEntry(
-					"ResetLevelInstancePropertyOverrides",
-					LOCTEXT("ResetLevelInstancePropertyOverrides", "Reset Override(s)"),
-					TAttribute<FText>(),
-					TAttribute<FSlateIcon>(),
-					UIAction);
+					Section.AddMenuEntry(
+						"ResetLevelInstancePropertyOverrides",
+						LOCTEXT("ResetLevelInstancePropertyOverrides", "Reset Override(s)"),
+						TAttribute<FText>(),
+						TAttribute<FSlateIcon>(),
+						UIAction);
 
-				return;
+					return;
+				}
 			}
 
 			if (SelectedActors.Num() > 0)
@@ -1033,6 +1040,8 @@ void FLevelInstanceEditorModule::StartupModule()
 		// Create a Behavior source for the default EdModeTools (when we aren't in the LevelInstanceEditorMode)
 		DefaultBehaviorSource = ULevelInstanceEditorMode::CreateDefaultModeBehaviorSource(GLevelEditorModeTools().GetInteractiveToolsContext());
 		GLevelEditorModeTools().GetInteractiveToolsContext()->InputRouter->RegisterSource(DefaultBehaviorSource.GetInterface());
+
+		RegisterLevelInstanceColumn();
 	}
 
 	ULevelInstanceSubsystem::RegisterPrimitiveColorHandler();
@@ -1059,11 +1068,42 @@ void FLevelInstanceEditorModule::ShutdownModule()
 
 	EditorLevelUtils::CanMoveActorToLevelDelegate.RemoveAll(this);
 
-	if (!IsRunningCommandlet() && GLevelEditorModeToolsIsValid())
+	if (!IsRunningCommandlet())
+	{
+		if (GLevelEditorModeToolsIsValid())
+		{
+			GLevelEditorModeTools().OnEditorModeIDChanged().RemoveAll(this);
+			GLevelEditorModeTools().GetInteractiveToolsContext()->InputRouter->DeregisterSource(DefaultBehaviorSource.GetInterface());
+			DefaultBehaviorSource = nullptr;
+		}
+		UnregisterLevelInstanceColumn();
+	}
+}
+
+TSharedRef<ISceneOutlinerColumn> FLevelInstanceEditorModule::CreateLevelInstanceColumn(ISceneOutliner& SceneOutliner) const
+{
+	return MakeShareable(new FLevelInstanceSceneOutlinerColumn(SceneOutliner));
+}
+
+void FLevelInstanceEditorModule::RegisterLevelInstanceColumn()
+{
+	if (GetDefault<ULevelInstanceSettings>()->IsPropertyOverrideEnabled())
+	{
+		FSceneOutlinerModule& SceneOutlinerModule = FModuleManager::LoadModuleChecked<FSceneOutlinerModule>("SceneOutliner");
+
+		FSceneOutlinerColumnInfo ColumnInfo(ESceneOutlinerColumnVisibility::Invisible, 8,
+			FCreateSceneOutlinerColumn::CreateRaw(this, &FLevelInstanceEditorModule::CreateLevelInstanceColumn),
+			true, TOptional<float>(), LOCTEXT("LevelInstanceColumnName", "Level Instance"));
+
+		SceneOutlinerModule.RegisterDefaultColumnType<FLevelInstanceSceneOutlinerColumn>(ColumnInfo);
+	}
+}
+
+void FLevelInstanceEditorModule::UnregisterLevelInstanceColumn()
+{
+	if (FSceneOutlinerModule* SceneOutlinerModulePtr = FModuleManager::GetModulePtr<FSceneOutlinerModule>("SceneOutliner"))
 	{
-		GLevelEditorModeTools().OnEditorModeIDChanged().RemoveAll(this);
-		GLevelEditorModeTools().GetInteractiveToolsContext()->InputRouter->DeregisterSource(DefaultBehaviorSource.GetInterface());
-		DefaultBehaviorSource = nullptr;
+		SceneOutlinerModulePtr->UnRegisterColumnType<FLevelInstanceSceneOutlinerColumn>();
 	}
 }
 
diff --git a/Engine/Source/Editor/LevelInstanceEditor/Private/LevelInstanceEditorModule.h b/Engine/Source/Editor/LevelInstanceEditor/Private/LevelInstanceEditorModule.h
--- a/Engine/Source/Editor/LevelInstanceEditor/Private/LevelInstanceEditorModule.h
+++ b/Engine/Source/Editor/LevelInstanceEditor/Private/LevelInstanceEditorModule.h
@@ -13,6 +13,8 @@ class ULevel;
 class IInputBehaviorSource;
 enum class EMapChangeType : uint8;
 class ILevelEditor;
+class ISceneOutliner;
+class ISceneOutlinerColumn;
 
 /**
  * The module holding all of the UI related pieces for LevelInstance management
@@ -57,6 +59,10 @@ private:
 	void OnLevelEditorCreated(TSharedPtr<ILevelEditor> InLevelEditor);
 	void RegisterToFirstLevelEditor();
 
+	void RegisterLevelInstanceColumn();
+	void UnregisterLevelInstanceColumn();
+	TSharedRef<ISceneOutlinerColumn> CreateLevelInstanceColumn(ISceneOutliner& SceneOutliner) const;
+
 	class FPropertyEditorPolicy : public PropertyEditorPolicy::IEditConstPolicy, 
 								  public PropertyEditorPolicy::IArchetypePolicy
 	{
diff --git a/Engine/Source/Editor/LevelInstanceEditor/Private/LevelInstanceSceneOutlinerColumn.cpp b/Engine/Source/Editor/LevelInstanceEditor/Private/LevelInstanceSceneOutlinerColumn.cpp
new file mode 100644
--- /dev/null
+++ b/Engine/Source/Editor/LevelInstanceEditor/Private/LevelInstanceSceneOutlinerColumn.cpp
@@ -0,0 +1,134 @@
+// Copyright Epic Games, Inc. All Rights Reserved.
+
+#include "LevelInstanceSceneOutlinerColumn.h"
+#include "Styling/SlateIconFinder.h"
+#include "Widgets/SNullWidget.h"
+#include "Widgets/SBoxPanel.h"
+#include "Widgets/SOverlay.h"
+#include "Widgets/Images/SImage.h"
+#include "GameFramework/Actor.h"
+#include "LevelInstance/LevelInstanceActor.h"
+#include "LevelInstance/LevelInstanceInterface.h"
+#include "ActorTreeItem.h"
+
+#define LOCTEXT_NAMESPACE "LevelInstanceColumn"
+
+namespace LevelInstanceColumnPrivate
+{
+	FName Name("Level Instance");
+
+	const FText ToolTipIsOverriden = LOCTEXT("IsOverridenTooltip", "Overriden");
+	const FText ToolTipIsOverridenAndContainsOverrides = LOCTEXT("IsOverridenAndContainsOverridesTooltip", "Overriden and overrides actor(s)");
+	const FText ToolTipContainsOverrides = LOCTEXT("ContainsOverridesTooltip", "Overrides actor(s)");
+
+	void GetBrushesAndToolTipForItem(FSceneOutlinerTreeItemRef TreeItem, const FSlateBrush*& OutHasOVerrideBrush, const FSlateBrush*& OutContainsOverrideBrush, const FText*& OutToolTipText)
+	{
+		OutToolTipText = nullptr;
+		OutHasOVerrideBrush = nullptr;
+		OutContainsOverrideBrush = nullptr;
+
+		if (const FActorTreeItem* ActorTreeITem = TreeItem->CastTo<FActorTreeItem>())
+		{
+			if (const AActor* Actor = ActorTreeITem->Actor.Get())
+			{				
+				if (const ILevelInstanceInterface* LevelInstance = Cast<ILevelInstanceInterface>(Actor); LevelInstance && LevelInstance->GetPropertyOverrideAsset() && !LevelInstance->IsEditing())
+				{
+					const bool bIsEditable = !Actor->IsInLevelInstance() || Actor->IsInEditLevelInstance();
+
+					OutContainsOverrideBrush = bIsEditable ? FAppStyle::GetBrush("LevelInstance.ColumnOverrideContainerInsideEditable") : FAppStyle::GetBrush("LevelInstance.ColumnOverrideContainerInside");
+
+					if (Actor->HasLevelInstancePropertyOverrides())
+					{
+						OutHasOVerrideBrush = Actor->HasEditableLevelLevelInstancePropertyOverrides() ? FAppStyle::GetBrush("LevelInstance.ColumnOverrideContainerHereEditable") : FAppStyle::GetBrush("LevelInstance.ColumnOverrideContainerHere");
+						OutToolTipText = &ToolTipIsOverridenAndContainsOverrides;
+					}
+					else
+					{
+						OutHasOVerrideBrush = bIsEditable ? FAppStyle::GetBrush("LevelInstance.ColumnOverrideContainerEditable") : FAppStyle::GetBrush("LevelInstance.ColumnOverrideContainer");
+						OutToolTipText = &ToolTipContainsOverrides;
+					}
+				}
+				else if (Actor->HasLevelInstancePropertyOverrides())
+				{
+					OutHasOVerrideBrush = Actor->HasEditableLevelLevelInstancePropertyOverrides() ? FAppStyle::GetBrush("LevelInstance.ColumnOverrideHereEditable") : FAppStyle::GetBrush("LevelInstance.ColumnOverrideHere");
+					OutToolTipText = &ToolTipIsOverriden;
+				}
+			}
+		}
+	}
+}
+
+FName FLevelInstanceSceneOutlinerColumn::GetID()
+{
+	return LevelInstanceColumnPrivate::Name;
+}
+
+SHeaderRow::FColumn::FArguments FLevelInstanceSceneOutlinerColumn::ConstructHeaderRowColumn()
+{
+	return SHeaderRow::Column(GetColumnID())
+		.FixedWidth(24.f)
+		.HAlignHeader(HAlign_Center)
+		.VAlignHeader(VAlign_Center)
+		.HAlignCell(HAlign_Center)
+		.VAlignCell(VAlign_Center)
+		.DefaultTooltip(FText::FromName(GetColumnID()))
+		[
+			SNew(SImage)
+			.ColorAndOpacity(FSlateColor::UseForeground())
+			.Image(FSlateIconFinder::FindIconBrushForClass(ALevelInstance::StaticClass()))
+		];
+}
+
+const TSharedRef<SWidget> FLevelInstanceSceneOutlinerColumn::ConstructRowWidget(FSceneOutlinerTreeItemRef TreeItem, const STableRow<FSceneOutlinerTreeItemPtr>& Row)
+{
+	if (const FActorTreeItem* ActorTreeITem = TreeItem->CastTo<FActorTreeItem>())
+	{
+		if (const AActor* Actor = ActorTreeITem->Actor.Get())
+		{
+			// First overlay slot is optional
+			return SNew(SOverlay)
+					+SOverlay::Slot()
+					.HAlign(HAlign_Center)
+					.VAlign(VAlign_Center)
+					[
+						SNew(SImage)
+						.ColorAndOpacity(FSlateColor::UseForeground())
+						.Image_Lambda([TreeItem]()
+						{
+							const FSlateBrush* OutHasOVerrideBrush;
+							const FSlateBrush* OutContainsOverrideBrush;
+							const FText* OutToolTipText;
+							LevelInstanceColumnPrivate::GetBrushesAndToolTipForItem(TreeItem, OutHasOVerrideBrush, OutContainsOverrideBrush, OutToolTipText);
+							return OutContainsOverrideBrush;
+						})
+					]
+					+SOverlay::Slot()
+					.HAlign(HAlign_Center)
+					.VAlign(VAlign_Center)
+					[
+						SNew(SImage)
+						.ColorAndOpacity(FSlateColor::UseForeground())
+						.Image_Lambda([TreeItem]()
+						{
+							const FSlateBrush* OutHasOVerrideBrush;
+							const FSlateBrush* OutContainsOverrideBrush;
+							const FText* OutToolTipText;
+							LevelInstanceColumnPrivate::GetBrushesAndToolTipForItem(TreeItem, OutHasOVerrideBrush, OutContainsOverrideBrush, OutToolTipText);
+							return OutHasOVerrideBrush;
+						})
+						.ToolTipText_Lambda([TreeItem]()
+						{ 
+							const FSlateBrush* OutHasOVerrideBrush;
+							const FSlateBrush* OutContainsOverrideBrush;
+							const FText* OutToolTipText;
+							LevelInstanceColumnPrivate::GetBrushesAndToolTipForItem(TreeItem, OutHasOVerrideBrush, OutContainsOverrideBrush, OutToolTipText);
+							return OutToolTipText ? *OutToolTipText : FText();
+						})
+					];
+			
+		}
+	}
+	return SNullWidget::NullWidget;
+}
+
+#undef LOCTEXT_NAMESPACE
\ No newline at end of file
diff --git a/Engine/Source/Editor/LevelInstanceEditor/Private/LevelInstanceSceneOutlinerColumn.h b/Engine/Source/Editor/LevelInstanceEditor/Private/LevelInstanceSceneOutlinerColumn.h
new file mode 100644
--- /dev/null
+++ b/Engine/Source/Editor/LevelInstanceEditor/Private/LevelInstanceSceneOutlinerColumn.h
@@ -0,0 +1,25 @@
+// Copyright Epic Games, Inc. All Rights Reserved.
+
+#pragma once
+
+#include "CoreMinimal.h"
+#include "SceneOutlinerFwd.h"
+#include "ISceneOutliner.h"
+#include "ISceneOutlinerColumn.h"
+#include "Widgets/Views/SHeaderRow.h"
+
+template<typename ItemType> class STableRow;
+class FLevelInstanceSceneOutlinerColumn : public ISceneOutlinerColumn
+{
+public:
+	FLevelInstanceSceneOutlinerColumn(ISceneOutliner& SceneOutliner) {}
+	virtual ~FLevelInstanceSceneOutlinerColumn() {}
+	static FName GetID();
+
+	//~ Begin ISceneOutlinerColumn Interface
+	virtual FName GetColumnID() override { return GetID(); }
+	virtual SHeaderRow::FColumn::FArguments ConstructHeaderRowColumn() override;
+	virtual const TSharedRef<SWidget> ConstructRowWidget(FSceneOutlinerTreeItemRef TreeItem, const STableRow<FSceneOutlinerTreeItemPtr>& Row) override;
+	virtual bool SupportsSorting() const override { return false; }
+	//~ End ISceneOutlinerColumn Interface
+};
\ No newline at end of file
