Level Instance: UX tweaks following feedback from Dan and Robb (validated by JS)
- Wording: Commit/Discard -> Save/Cancel
- Details panel buttons cleanup
- Viewport overlay updated now with Save/Cancel instead of Exit only
- Menus updated to reflect new wording
- Remove: Move Selection To and Set As Current menus as those actions can be achieved through the Actors menu and Actor Context viewport overlay

#rnx

diff --git a/Engine/Source/Editor/LevelInstanceEditor/Private/LevelInstanceActorDetails.cpp b/Engine/Source/Editor/LevelInstanceEditor/Private/LevelInstanceActorDetails.cpp
--- a/Engine/Source/Editor/LevelInstanceEditor/Private/LevelInstanceActorDetails.cpp
+++ b/Engine/Source/Editor/LevelInstanceEditor/Private/LevelInstanceActorDetails.cpp
@@ -1,12 +1,14 @@
 // Copyright Epic Games, Inc. All Rights Reserved.
 
 #include "LevelInstanceActorDetails.h"
+#include "LevelInstance/LevelInstanceSettings.h"
+#include "LevelInstance/LevelInstanceSubsystem.h"
 #include "DetailLayoutBuilder.h"
 #include "DetailCategoryBuilder.h"
 #include "DetailWidgetRow.h"
 #include "UObject/WeakInterfacePtr.h"
 #include "Widgets/Input/SButton.h"
-#include "Widgets/Input/SMultiLineEditableTextBox.h"
+#include "Widgets/Text/STextBlock.h"
 #include "ScopedTransaction.h"
 
 #include "Engine/World.h"
@@ -14,68 +16,218 @@
 
 #define LOCTEXT_NAMESPACE "FLevelInstanceActorDetails"
 
+struct FLevelInstanceActorDetailsHelper
+{
+	static void ResetPropertyOverrides(ILevelInstanceInterface* LevelInstance)
+	{
+		check(LevelInstance);
+		LevelInstance->GetLevelInstanceSubsystem()->ResetPropertyOverrides(LevelInstance);
+	}
+};
+
 namespace LevelInstanceActorDetailsCallbacks
 {
-	static bool IsEditCommitButtonEnabled(TWeakInterfacePtr<ILevelInstanceInterface> LevelInstancePtr)
+	static bool IsEditButtonEnabled(TWeakInterfacePtr<ILevelInstanceInterface> LevelInstancePtr)
 	{
 		if (ILevelInstanceInterface* LevelInstance = LevelInstancePtr.Get())
 		{
-			return LevelInstance->CanEnterEdit() || LevelInstance->CanExitEdit();
+			return LevelInstance->CanEnterEdit();
 		}
 
 		return false;
 	}
 
-	static FText GetEditCommitButtonText(TWeakInterfacePtr<ILevelInstanceInterface> LevelInstancePtr)
+	static FText GetEditButtonTooltip(TWeakInterfacePtr<ILevelInstanceInterface> LevelInstancePtr)
 	{
+		FText Reason;
 		if (ILevelInstanceInterface* LevelInstance = LevelInstancePtr.Get())
 		{
-			if (LevelInstance->CanExitEdit())
+			if (!LevelInstance->IsEditing())
 			{
-				return LOCTEXT("CommitChanges", "Commit Changes");
+				if (!LevelInstance->CanEnterEdit(&Reason))
+				{
+					return Reason;
+				}
+
+				return LOCTEXT("EditButtonToolTip", "Edit level instance source level");
 			}
 		}
+		return FText::GetEmpty();
+	}
+
+	static EVisibility GetEditButtonVisibility(TWeakInterfacePtr<ILevelInstanceInterface> LevelInstancePtr)
+	{
+		if (ILevelInstanceInterface* LevelInstance = LevelInstancePtr.Get())
+		{
+			return (!LevelInstance->IsEditing() && !LevelInstance->IsEditingPropertyOverrides()) ? EVisibility::Visible : EVisibility::Collapsed;
+		}
+
+		return EVisibility::Collapsed;
+	}
 		
-		return LOCTEXT("Edit", "Edit");
+	static FReply OnEditButtonClicked(TWeakInterfacePtr<ILevelInstanceInterface> LevelInstancePtr)
+	{
+		if (ILevelInstanceInterface* LevelInstance = LevelInstancePtr.Get())
+		{
+			if (LevelInstance->CanEnterEdit())
+			{
+				LevelInstance->EnterEdit();
+			}
+		}
+		return FReply::Handled();
 	}
 
-	static FText GetEditCommitReasonText(TWeakInterfacePtr<ILevelInstanceInterface> LevelInstancePtr)
+	static bool IsOverrideButtonEnabled(TWeakInterfacePtr<ILevelInstanceInterface> LevelInstancePtr)
+	{
+		if (ILevelInstanceInterface* LevelInstance = LevelInstancePtr.Get())
+		{
+			return LevelInstance->CanEnterEditPropertyOverrides();
+		}
+
+		return false;
+	}
+
+	static EVisibility GetOverrideButtonVisibility(TWeakInterfacePtr<ILevelInstanceInterface> LevelInstancePtr)
+	{
+		if (ULevelInstanceSettings::Get()->IsPropertyOverrideEnabled())
+		{
+			return GetEditButtonVisibility(LevelInstancePtr);
+		}
+
+		return EVisibility::Collapsed;
+	}
+
+	static FText GetOverrideButtonTooltip(TWeakInterfacePtr<ILevelInstanceInterface> LevelInstancePtr)
 	{
 		FText Reason;
 		if (ILevelInstanceInterface* LevelInstance = LevelInstancePtr.Get())
 		{
-			if (!LevelInstance->IsEditing())
+			if (!LevelInstance->IsEditingPropertyOverrides())
+			{
+				if (!LevelInstance->CanEnterEditPropertyOverrides(&Reason))
+				{
+					return Reason;
+				}
+
+				return LOCTEXT("OverrideButtonToolTip", "Override properties on level instance actors");
+			}
+		}
+		return FText::GetEmpty();
+	}
+
+	static FReply OnOverrideButtonClicked(TWeakInterfacePtr<ILevelInstanceInterface> LevelInstancePtr)
+	{
+		if (ILevelInstanceInterface* LevelInstance = LevelInstancePtr.Get())
+		{
+			if (LevelInstance->CanEnterEditPropertyOverrides())
 			{
-				LevelInstance->CanEnterEdit(&Reason);
-				return Reason;
+				LevelInstance->EnterEditPropertyOverrides();
 			}
+		}
+		return FReply::Handled();
+	}
 
-			LevelInstance->CanExitEdit(/*bDiscardEdits=*/false, &Reason);
+	static bool IsResetOverridesButtonEnabled(TWeakInterfacePtr<ILevelInstanceInterface> LevelInstancePtr)
+	{
+		if (ILevelInstanceInterface* LevelInstance = LevelInstancePtr.Get())
+		{
+			return LevelInstance->GetPropertyOverrideAsset() != nullptr;
 		}
-		return Reason;
+
+		return false;
 	}
 
-	static EVisibility GetEditCommitReasonVisibility(TWeakInterfacePtr<ILevelInstanceInterface> LevelInstancePtr)
+	static FText GetResetOverridesButtonTooltip(TWeakInterfacePtr<ILevelInstanceInterface> LevelInstancePtr)
+	{
+		return LOCTEXT("ResetOverrideButtonToolTip", "Reset property overrides on level instance actor"); 
+	}
+
+	static FReply OnResetOverridesButtonClicked(TWeakInterfacePtr<ILevelInstanceInterface> LevelInstancePtr)
 	{
 		if (ILevelInstanceInterface* LevelInstance = LevelInstancePtr.Get())
 		{
-			return IsEditCommitButtonEnabled(LevelInstance) ? EVisibility::Collapsed : EVisibility::Visible;
+			FLevelInstanceActorDetailsHelper::ResetPropertyOverrides(LevelInstance);
+		}
+
+		return FReply::Handled();
+	}
+
+	static EVisibility GetResetOverridesButtonVisibility(TWeakInterfacePtr<ILevelInstanceInterface> LevelInstancePtr)
+	{
+		if (ILevelInstanceInterface* LevelInstance = LevelInstancePtr.Get())
+		{
+			AActor* Actor = CastChecked<AActor>(LevelInstance);
+			return ULevelInstanceSettings::Get()->IsPropertyOverrideEnabled() && !LevelInstance->IsEditing() && !LevelInstance->IsEditingPropertyOverrides() && (!Actor->IsInLevelInstance() || Actor->IsInEditLevelInstance()) ? EVisibility::Visible : EVisibility::Collapsed;
+		}
+
+		return EVisibility::Collapsed;
+	}
+
+	static EVisibility GetSaveCancelButtonVisibility(TWeakInterfacePtr<ILevelInstanceInterface> LevelInstancePtr)
+	{
+		if (ILevelInstanceInterface* LevelInstance = LevelInstancePtr.Get())
+		{
+			return (LevelInstance->IsEditing() || LevelInstance->IsEditingPropertyOverrides()) ? EVisibility::Visible : EVisibility::Collapsed;
 		}
 
 		return EVisibility::Collapsed;
 	}
+		
+	static bool IsSaveCancelButtonEnabled(TWeakInterfacePtr<ILevelInstanceInterface> LevelInstancePtr, bool bDiscard)
+	{
+		if (ILevelInstanceInterface* LevelInstance = LevelInstancePtr.Get())
+		{
+			if (LevelInstance->IsEditing())
+			{
+				return LevelInstance->CanExitEdit(bDiscard);
+			}
+			else if (LevelInstance->IsEditingPropertyOverrides())
+			{
+				return LevelInstance->CanExitEditPropertyOverrides(bDiscard);
+			}
+		}
+
+		return false;
+	}
 
-	static FReply OnEditCommitButtonClicked(TWeakInterfacePtr<ILevelInstanceInterface> LevelInstancePtr)
+	static FText GetSaveCancelButtonTooltip(TWeakInterfacePtr<ILevelInstanceInterface> LevelInstancePtr, bool bDiscard)
 	{
+		FText Reason;
 		if (ILevelInstanceInterface* LevelInstance = LevelInstancePtr.Get())
 		{
-			if (LevelInstance->CanExitEdit())
+			if (LevelInstance->IsEditing())
 			{
-				LevelInstance->ExitEdit();
+				if (!LevelInstance->CanExitEdit(bDiscard, &Reason))
+				{
+					return Reason;
+				}
+
+				return bDiscard ? LOCTEXT("CancelButtonToolTip", "Cancel edits and exit") : LOCTEXT("SaveButtonToolTip", "Save edits and exit");
 			}
-			else if (LevelInstance->CanEnterEdit())
+			else if (LevelInstance->IsEditingPropertyOverrides())
 			{
-				LevelInstance->EnterEdit();
+				if (!LevelInstance->CanExitEditPropertyOverrides(bDiscard, &Reason))
+				{
+					return Reason;
+				}
+
+				return bDiscard ? LOCTEXT("CancelOverrideButtonToolTip", "Cancel overrides and exit") : LOCTEXT("SaveOverrideButtonToolTip", "Save overrides and exit");
+			}
+		}
+		return FText::GetEmpty();
+	}
+
+	static FReply OnSaveCancelButtonClicked(TWeakInterfacePtr<ILevelInstanceInterface> LevelInstancePtr, bool bDiscard)
+	{
+		if (ILevelInstanceInterface* LevelInstance = LevelInstancePtr.Get())
+		{
+			if (LevelInstance->IsEditing())
+			{
+				LevelInstance->ExitEdit(bDiscard);
+			}
+			else if (LevelInstance->IsEditingPropertyOverrides())
+			{
+				LevelInstance->ExitEditPropertyOverrides(bDiscard);
 			}
 		}
 		return FReply::Handled();
@@ -114,42 +266,104 @@ void FLevelInstanceActorDetails::CustomizeDetails(IDetailLayoutBuilder& DetailBu
 		DetailBuilder.HideProperty("WorldAsset");
 	}
 
-	IDetailCategoryBuilder& LevelInstanceEditingCategory = DetailBuilder.EditCategory("Level Edit", FText::GetEmpty(), ECategoryPriority::Transform);
+	IDetailCategoryBuilder& LevelInstanceEditingCategory = DetailBuilder.EditCategory("LevelInstanceEdit", LOCTEXT("LevelInstanceEditCategory", "Level Instance"), ECategoryPriority::Transform);
 
-	LevelInstanceEditingCategory.AddCustomRow(FText::GetEmpty())
-	[
-		SNew(SVerticalBox)
-		+ SVerticalBox::Slot()
-		.AutoHeight()
+	LevelInstanceEditingCategory.AddCustomRow(FText::GetEmpty()).ValueContent()
 		[
 			SNew(SHorizontalBox)
 			+ SHorizontalBox::Slot()
-			.FillWidth(1)
+			.AutoWidth()
+			.Padding(4.0f, 8.0f, 4.0f, 8.0f)
+			.VAlign(VAlign_Center)
 			[
-				SNew(SMultiLineEditableTextBox)
-				.Visibility_Static(&LevelInstanceActorDetailsCallbacks::GetEditCommitReasonVisibility, LevelInstance)
-				.Font(DetailBuilder.GetDetailFontBold())
-				.BackgroundColor(TAttribute<FSlateColor>::Create(TAttribute<FSlateColor>::FGetter::CreateLambda([]() { return FAppStyle::GetColor("ErrorReporting.WarningBackgroundColor"); })))
-				.Text_Static(&LevelInstanceActorDetailsCallbacks::GetEditCommitReasonText, LevelInstance)
-				.AutoWrapText(true)
-				.IsReadOnly(true)
+				SNew(SButton)
+				.IsEnabled_Static(&LevelInstanceActorDetailsCallbacks::IsEditButtonEnabled, LevelInstance)
+				.ToolTipText_Static(&LevelInstanceActorDetailsCallbacks::GetEditButtonTooltip, LevelInstance)
+				.Visibility_Static(&LevelInstanceActorDetailsCallbacks::GetEditButtonVisibility, LevelInstance)
+				.HAlign(HAlign_Center)
+				.VAlign(VAlign_Center)
+				.OnClicked_Static(&LevelInstanceActorDetailsCallbacks::OnEditButtonClicked, LevelInstance)
+				[
+					SNew(STextBlock)
+					.Font(IDetailLayoutBuilder::GetDetailFont())
+					.Text(LOCTEXT("EditText", "Edit"))
+				]
+			]
+			+ SHorizontalBox::Slot()
+			.AutoWidth()
+			.Padding(4.0f, 8.0f, 4.0f, 8.0f)
+			.VAlign(VAlign_Center)
+			[
+				SNew(SButton)
+				.IsEnabled_Static(&LevelInstanceActorDetailsCallbacks::IsOverrideButtonEnabled, LevelInstance)
+				.ToolTipText_Static(&LevelInstanceActorDetailsCallbacks::GetOverrideButtonTooltip, LevelInstance)
+				.Visibility_Static(&LevelInstanceActorDetailsCallbacks::GetOverrideButtonVisibility, LevelInstance)
+				.HAlign(HAlign_Center)
+				.VAlign(VAlign_Center)
+				.OnClicked_Static(&LevelInstanceActorDetailsCallbacks::OnOverrideButtonClicked, LevelInstance)
+				[
+					SNew(STextBlock)
+					.Font(IDetailLayoutBuilder::GetDetailFont())
+					.Text(LOCTEXT("OverrideText", "Override"))
+				]
+			]
+			+ SHorizontalBox::Slot()
+			.AutoWidth()
+			.Padding(4.0f, 8.0f, 8.0f, 8.0f)
+			.VAlign(VAlign_Center)
+			[
+				SNew(SButton)
+				.IsEnabled_Static(&LevelInstanceActorDetailsCallbacks::IsResetOverridesButtonEnabled, LevelInstance)
+				.ToolTipText_Static(&LevelInstanceActorDetailsCallbacks::GetResetOverridesButtonTooltip, LevelInstance)
+				.Visibility_Static(&LevelInstanceActorDetailsCallbacks::GetResetOverridesButtonVisibility, LevelInstance)
+				.HAlign(HAlign_Center)
+				.VAlign(VAlign_Center)
+				.OnClicked_Static(&LevelInstanceActorDetailsCallbacks::OnResetOverridesButtonClicked, LevelInstance)
+				[
+					SNew(STextBlock)
+					.Font(IDetailLayoutBuilder::GetDetailFont())
+					.Text(LOCTEXT("ResetOverrideText", "Reset Overrides"))
+				]
+			]
+			+ SHorizontalBox::Slot()
+			.AutoWidth()
+			.Padding(4.0f, 8.0f, 4.0f, 8.0f)
+			.VAlign(VAlign_Center)
+			[
+				SNew(SButton)
+				.ButtonStyle(FAppStyle::Get(), "PrimaryButton")
+				.IsEnabled_Static(&LevelInstanceActorDetailsCallbacks::IsSaveCancelButtonEnabled, LevelInstance, false)
+				.ToolTipText_Static(&LevelInstanceActorDetailsCallbacks::GetSaveCancelButtonTooltip, LevelInstance, false)
+				.Visibility_Static(&LevelInstanceActorDetailsCallbacks::GetSaveCancelButtonVisibility, LevelInstance)
+				.HAlign(HAlign_Center)
+				.VAlign(VAlign_Center)
+				.OnClicked_Static(&LevelInstanceActorDetailsCallbacks::OnSaveCancelButtonClicked, LevelInstance, false)
+				[
+					SNew(STextBlock)
+					.Font(IDetailLayoutBuilder::GetDetailFont())
+					.Text(LOCTEXT("SaveText", "Save"))
+				]
 			]
-		]
-		+ SVerticalBox::Slot()
-		.AutoHeight()
-		[
-			SNew(SHorizontalBox)
 			+ SHorizontalBox::Slot()
-			.FillWidth(1)
+			.AutoWidth()
+			.Padding(4.0f, 8.0f, 4.0f, 8.0f)
+			.VAlign(VAlign_Center)
 			[
 				SNew(SButton)
-				.IsEnabled_Static(&LevelInstanceActorDetailsCallbacks::IsEditCommitButtonEnabled, LevelInstance)
-				.Text_Static(&LevelInstanceActorDetailsCallbacks::GetEditCommitButtonText, LevelInstance)
+				.IsEnabled_Static(&LevelInstanceActorDetailsCallbacks::IsSaveCancelButtonEnabled, LevelInstance, true)
+				.ToolTipText_Static(&LevelInstanceActorDetailsCallbacks::GetSaveCancelButtonTooltip, LevelInstance, true)
+				.Visibility_Static(&LevelInstanceActorDetailsCallbacks::GetSaveCancelButtonVisibility, LevelInstance)
+				.Text(LOCTEXT("DiscardText", "Discard"))
 				.HAlign(HAlign_Center)
-				.OnClicked_Static(&LevelInstanceActorDetailsCallbacks::OnEditCommitButtonClicked, LevelInstance)
+				.VAlign(VAlign_Center)
+				.OnClicked_Static(&LevelInstanceActorDetailsCallbacks::OnSaveCancelButtonClicked, LevelInstance, true)
+				[
+					SNew(STextBlock)
+					.Font(IDetailLayoutBuilder::GetDetailFont())
+					.Text(LOCTEXT("CancelText", "Cancel"))
+				]
 			]
-		]
-	];
+		];
 }
 
 
diff --git a/Engine/Source/Editor/LevelInstanceEditor/Private/LevelInstanceEditorModeToolkit.cpp b/Engine/Source/Editor/LevelInstanceEditor/Private/LevelInstanceEditorModeToolkit.cpp
--- a/Engine/Source/Editor/LevelInstanceEditor/Private/LevelInstanceEditorModeToolkit.cpp
+++ b/Engine/Source/Editor/LevelInstanceEditor/Private/LevelInstanceEditorModeToolkit.cpp
@@ -17,6 +17,66 @@ class FAssetEditorModeUILayer;
 
 #define LOCTEXT_NAMESPACE "LevelInstanceEditorModeToolkit"
 
+struct FLevelInstanceEditorModeToolkitHelper
+{
+	static FText GetToolkitDisplayText(ULevelInstanceSubsystem* LevelInstanceSubsystem)
+	{
+		if (LevelInstanceSubsystem->GetEditingLevelInstance())
+		{
+			return LOCTEXT("LevelInstanceEditToolkitDisplayText", "Level Instance Edit");
+		}
+		else if(LevelInstanceSubsystem->GetEditingPropertyOverridesLevelInstance())
+		{
+			return LOCTEXT("LevelInstanceOverrideToolkitDisplayText", "Level Instance Override");
+		} 
+
+		return FText();
+	}
+
+	static FText GetToolkitSaveCancelButtonTooltipText(ULevelInstanceSubsystem* LevelInstanceSubsystem, bool bDiscard)
+	{
+		if (LevelInstanceSubsystem->GetEditingLevelInstance())
+		{
+			return bDiscard ? LOCTEXT("LevelInstanceCancelEditToolkitTooltip", "Cancel edits and exit") : LOCTEXT("LevelInstanceSaveEditToolkitTooltip", "Save edits and exit");
+		}
+		else if (LevelInstanceSubsystem->GetEditingPropertyOverridesLevelInstance())
+		{
+			return bDiscard ? LOCTEXT("LevelInstanceCancelOverrideToolkitTooltip", "Cancel overrides and exit") : LOCTEXT("LevelInstanceSaveOverrideToolkitTooltip", "Save overrides and exit");
+		}
+
+		return FText();
+	}
+
+	static FReply OnSaveCancelButtonClicked(ULevelInstanceSubsystem* LevelInstanceSubsystem, bool bDiscard)
+	{
+		if (ILevelInstanceInterface* LevelInstance = LevelInstanceSubsystem->GetEditingLevelInstance())
+		{
+			LevelInstance->ExitEdit(bDiscard);
+		}
+		else if (ILevelInstanceInterface* LevelInstanceOverride = LevelInstanceSubsystem->GetEditingPropertyOverridesLevelInstance())
+		{
+			LevelInstanceOverride->ExitEditPropertyOverrides(bDiscard);
+		}
+
+		return FReply::Handled();
+	}
+
+	static bool IsCancelButtonEnabled(ULevelInstanceSubsystem* LevelInstanceSubsystem)
+	{
+		if (ILevelInstanceInterface* LevelInstance = LevelInstanceSubsystem->GetEditingLevelInstance())
+		{
+			return LevelInstance->CanExitEdit(true);
+		}
+		else if (ILevelInstanceInterface* LevelInstanceOverride = LevelInstanceSubsystem->GetEditingPropertyOverridesLevelInstance())
+		{
+			return LevelInstanceOverride->CanExitEditPropertyOverrides(true);
+		}
+
+		return false;
+	}
+};
+
+
 FLevelInstanceEditorModeToolkit::FLevelInstanceEditorModeToolkit()
 {
 }
@@ -51,7 +111,7 @@ void FLevelInstanceEditorModeToolkit::Init(const TSharedPtr<IToolkitHost>& InitT
 			+SHorizontalBox::Slot()
 			.AutoWidth()
 			.VAlign(VAlign_Center)
-			.Padding(FMargin(0.f, 0.f, 8.f, 0.f))
+			.Padding(FMargin(4.f, 0.f, 0.f, 0.f))
 			[
 				SNew(SImage)
 				.Image(FSlateIconFinder::FindIconBrushForClass(ALevelInstance::StaticClass()))
@@ -59,32 +119,34 @@ void FLevelInstanceEditorModeToolkit::Init(const TSharedPtr<IToolkitHost>& InitT
 			+SHorizontalBox::Slot()
 			.AutoWidth()
 			.VAlign(VAlign_Center)
-			.Padding(FMargin(0.f, 0.f, 8.f, 0.f))
+			.Padding(FMargin(8.f, 0.f, 0.f, 0.f))
 			[
 				SNew(STextBlock)
-				.Text_Lambda([LevelInstanceSubsystem]()
-				{
-					return LevelInstanceSubsystem->GetToolKitDisplayText();
-				})
+				.Text_Static(&FLevelInstanceEditorModeToolkitHelper::GetToolkitDisplayText, LevelInstanceSubsystem)
 			]
 			+SHorizontalBox::Slot()
 			.AutoWidth()
-			.Padding(FMargin(2.0, 0.f, 0.f, 0.f))
+			.Padding(FMargin(8.0, 0.f, 0.f, 0.f))
 			[
 				SNew(SButton)
 				.ButtonStyle(FAppStyle::Get(), "PrimaryButton")
 				.TextStyle(FAppStyle::Get(), "DialogButtonText")
-				.Text(LOCTEXT("ExitEdit", "Exit"))
-				.ToolTipText_Lambda([LevelInstanceSubsystem]()
-				{
-					return LevelInstanceSubsystem->GetToolKitExitToolTip();
-				})
+				.Text(LOCTEXT("SaveButtonText", "Save"))
+				.ToolTipText_Static(&FLevelInstanceEditorModeToolkitHelper::GetToolkitSaveCancelButtonTooltipText, LevelInstanceSubsystem, false)
+				.HAlign(HAlign_Center)
+				.OnClicked_Static(&FLevelInstanceEditorModeToolkitHelper::OnSaveCancelButtonClicked, LevelInstanceSubsystem, false)
+			]
+			+SHorizontalBox::Slot()
+			.AutoWidth()
+			.Padding(FMargin(4.0, 0.f, 4.f, 0.f))
+			[
+				SNew(SButton)
+				.TextStyle(FAppStyle::Get(), "DialogButtonText")
+				.Text(LOCTEXT("CancelButtonText", "Cancel"))
+				.ToolTipText_Static(&FLevelInstanceEditorModeToolkitHelper::GetToolkitSaveCancelButtonTooltipText, LevelInstanceSubsystem, true)
 				.HAlign(HAlign_Center)
-				.OnClicked_Lambda([LevelInstanceSubsystem]() 
-				{ 
-					LevelInstanceSubsystem->ToolKitExit();
-					return FReply::Handled(); 
-				})
+				.OnClicked_Static(&FLevelInstanceEditorModeToolkitHelper::OnSaveCancelButtonClicked, LevelInstanceSubsystem, true)
+				.IsEnabled_Static(&FLevelInstanceEditorModeToolkitHelper::IsCancelButtonEnabled, LevelInstanceSubsystem)
 			]
 		]
 	];
diff --git a/Engine/Source/Editor/LevelInstanceEditor/Private/LevelInstanceEditorModule.cpp b/Engine/Source/Editor/LevelInstanceEditor/Private/LevelInstanceEditorModule.cpp
--- a/Engine/Source/Editor/LevelInstanceEditor/Private/LevelInstanceEditorModule.cpp
+++ b/Engine/Source/Editor/LevelInstanceEditor/Private/LevelInstanceEditorModule.cpp
@@ -63,11 +63,20 @@ struct FLevelInstanceMenuUtils
 {
 	static FToolMenuSection& CreateLevelSection(UToolMenu* Menu)
 	{
-		const FName LevelSectionName = TEXT("Level");
-		FToolMenuSection* SectionPtr = Menu->FindSection(LevelSectionName);
+		return CreateSection(Menu, FName("Level"), LOCTEXT("LevelSectionLabel", "Level"));
+	}
+	
+	static FToolMenuSection& CreateCurrentEditSection(UToolMenu* Menu)
+	{
+		return CreateSection(Menu, FName("CurrentEdit"), LOCTEXT("CurrentEditSectionLabel", "Current Edit"));
+	}
+
+	static FToolMenuSection& CreateSection(UToolMenu* Menu, FName SectionName, const FText& SectionText)
+	{
+		FToolMenuSection* SectionPtr = Menu->FindSection(SectionName);
 		if (!SectionPtr)
 		{
-			SectionPtr = &(Menu->AddSection(LevelSectionName, LOCTEXT("LevelSectionLabel", "Level")));
+			SectionPtr = &(Menu->AddSection(SectionName, SectionText));
 		}
 		FToolMenuSection& Section = *SectionPtr;
 		return Section;
@@ -92,7 +101,8 @@ struct FLevelInstanceMenuUtils
 		FText EntryLabel = bSingleEntry ? LOCTEXT("EditLevelInstances", "Edit") : FText::FromString(LevelInstance->GetWorldAsset().GetAssetName());
 		if (bCanEdit)
 		{
-			EntryDesc = FText::Format(LOCTEXT("LevelInstanceName", "{0}:{1}"), FText::FromString(LevelInstanceActor->GetActorLabel()), FText::FromString(LevelInstance->GetWorldAssetPackage()));
+			FText EntryActionDesc = LOCTEXT("EditLevelInstancesPropertyTooltip", "Edit this level. Your changes will be applied to the level asset and to all other level instances based on it.");
+			EntryDesc = FText::Format(LOCTEXT("LevelInstanceName", "{0}\n\nActor name: {1}\nAsset path: {2}"), EntryActionDesc, FText::FromString(LevelInstanceActor->GetActorLabel()), FText::FromString(LevelInstance->GetWorldAssetPackage()));
 		}
 		Section.AddMenuEntry(NAME_None, EntryLabel, EntryDesc, FSlateIcon(), LevelInstanceEditAction);
 	}
@@ -125,7 +135,8 @@ struct FLevelInstanceMenuUtils
 		FText EntryLabel = bSingleEntry ? LOCTEXT("OverrideLevelInstances", "Override") : FText::FromString(LevelInstance->GetWorldAsset().GetAssetName());
 		if (bCanEdit)
 		{
-			EntryDesc = FText::Format(LOCTEXT("OverrideLevelInstanceName", "{0}:{1}"), FText::FromString(LevelInstanceActor->GetActorLabel()), FText::FromString(LevelInstance->GetWorldAssetPackage()));
+			FText EntryActionDesc = LOCTEXT("EditLevelInstancesPropertyOverridesTooltip", "Edit only this level instance, without changing the level asset or any other level instances.");
+			EntryDesc = FText::Format(LOCTEXT("OverrideLevelInstanceName", "{0}\n\nActor name: {1}\nAsset path: {2}"), EntryActionDesc, FText::FromString(LevelInstanceActor->GetActorLabel()), FText::FromString(LevelInstance->GetWorldAssetPackage()));
 		}
 		Section.AddMenuEntry(NAME_None, EntryLabel, EntryDesc, FSlateIcon(), LevelInstanceEditAction);
 	}
@@ -167,7 +178,7 @@ struct FLevelInstanceMenuUtils
 				Section.AddSubMenu(
 					"EditLevelInstances",
 					LOCTEXT("EditLevelInstances", "Edit"),
-					TAttribute<FText>(),
+					LOCTEXT("EditLevelInstancesPropertyTooltip", "Edit this level. Your changes will be applied to the level asset and to all other level instances based on it."),
 					FNewToolMenuDelegate::CreateStatic(&CreateEditSubMenu, MoveTemp(LevelInstanceHierarchy), ContextActor)
 				);
 			}
@@ -202,14 +213,14 @@ struct FLevelInstanceMenuUtils
 				Section.AddSubMenu(
 					"PropertyOverrideLevelInstances",
 					LOCTEXT("EditLevelInstancesPropertyOverrides", "Override"),
-					TAttribute<FText>(),
+					LOCTEXT("EditLevelInstancesPropertyOverridesTooltip", "Edit only this level instance, without changing the level asset or any other level instances."),
 					FNewToolMenuDelegate::CreateStatic(&CreateEditPropertyOverridesSubMenu, MoveTemp(LevelInstanceHierarchy), ContextActor)
 				);
 			}
 		}
 	}
 
-	static void CreateCommitDiscardMenu(UToolMenu* Menu, AActor* ContextActor)
+	static void CreateSaveCancelMenu(UToolMenu* Menu, AActor* ContextActor)
 	{
 		ILevelInstanceInterface* LevelInstanceEdit = nullptr;
 		if (ContextActor)
@@ -245,103 +256,42 @@ struct FLevelInstanceMenuUtils
 						
 		if (LevelInstanceEdit)
 		{
-			FToolMenuSection& Section = CreateLevelSection(Menu);
+			FToolMenuSection& Section = CreateCurrentEditSection(Menu);
 			if (LevelInstanceEdit->IsEditingPropertyOverrides())
 			{
-				FText CommitTooltip;
+				FText CommitTooltip = LOCTEXT("LevelInstanceCommitPropertyOverridesTooltip", "Stop overriding this level instance and save any changes you've made.");
 				const bool bCanCommit = LevelInstanceEdit->CanExitEditPropertyOverrides(/*bDiscardEdits=*/false, &CommitTooltip);
 
 				FToolUIAction CommitAction;
 				CommitAction.ExecuteAction.BindLambda([LevelInstanceEdit](const FToolMenuContext&) { LevelInstanceEdit->ExitEditPropertyOverrides(/*bDiscardEdits=*/false); });
 				CommitAction.CanExecuteAction.BindLambda([bCanCommit](const FToolMenuContext&) { return bCanCommit; });
-				Section.AddMenuEntry(NAME_None, LOCTEXT("LevelInstanceCommitPropertyOverridesLabel", "Commit Override(s)"), CommitTooltip, FSlateIcon(), CommitAction);
+				Section.AddMenuEntry(NAME_None, LOCTEXT("LevelInstanceSavePropertyOverridesLabel", "Save Override(s)"), CommitTooltip, FSlateIcon(), CommitAction);
 
-				FText DiscardTooltip;
+				FText DiscardTooltip = LOCTEXT("LevelInstanceDiscardPropertyOverridesTooltip", "Stop overriding this level instance and discard any changes you've made.");
 				const bool bCanDiscard = LevelInstanceEdit->CanExitEditPropertyOverrides(/*bDiscardEdits=*/true, &DiscardTooltip);
 
 				FToolUIAction DiscardAction;
 				DiscardAction.ExecuteAction.BindLambda([LevelInstanceEdit](const FToolMenuContext&) { LevelInstanceEdit->ExitEditPropertyOverrides(/*bDiscardEdits=*/true); });
 				DiscardAction.CanExecuteAction.BindLambda([bCanDiscard](const FToolMenuContext&) { return bCanDiscard; });
-				Section.AddMenuEntry(NAME_None, LOCTEXT("LevelInstanceDiscardPropertyOverridesLabel", "Discard Override(s)"), DiscardTooltip, FSlateIcon(), DiscardAction);
+				Section.AddMenuEntry(NAME_None, LOCTEXT("LevelInstanceCancelPropertyOverridesLabel", "Cancel Override(s)"), DiscardTooltip, FSlateIcon(), DiscardAction);
 			}
 			else
 			{
-				FText CommitTooltip;
+				FText CommitTooltip = LOCTEXT("LevelInstanceCommitTooltip", "Stop editing this level and save any changes you've made.");
 				const bool bCanCommit = LevelInstanceEdit->CanExitEdit(/*bDiscardEdits=*/false, &CommitTooltip);
 
 				FToolUIAction CommitAction;
 				CommitAction.ExecuteAction.BindLambda([LevelInstanceEdit](const FToolMenuContext&) { LevelInstanceEdit->ExitEdit(/*bDiscardEdits=*/false); });
 				CommitAction.CanExecuteAction.BindLambda([bCanCommit](const FToolMenuContext&) { return bCanCommit; });
-				Section.AddMenuEntry(NAME_None, LOCTEXT("LevelInstanceCommitLabel", "Commit"), CommitTooltip, FSlateIcon(), CommitAction);
+				Section.AddMenuEntry(NAME_None, LOCTEXT("LevelInstanceSaveLabel", "Save"), CommitTooltip, FSlateIcon(), CommitAction);
 
-				FText DiscardTooltip;
+				FText DiscardTooltip = LOCTEXT("LevelInstanceDiscardTooltip", "Stop editing this level and discard any changes you've made.");
 				const bool bCanDiscard = LevelInstanceEdit->CanExitEdit(/*bDiscardEdits=*/true, &DiscardTooltip);
 
 				FToolUIAction DiscardAction;
 				DiscardAction.ExecuteAction.BindLambda([LevelInstanceEdit](const FToolMenuContext&) { LevelInstanceEdit->ExitEdit(/*bDiscardEdits=*/true); });
 				DiscardAction.CanExecuteAction.BindLambda([bCanDiscard](const FToolMenuContext&) { return bCanDiscard; });
-				Section.AddMenuEntry(NAME_None, LOCTEXT("LevelInstanceDiscardLabel", "Discard"), DiscardTooltip, FSlateIcon(), DiscardAction);
-			}
-		}
-	}
-
-	static void CreateSetCurrentMenu(UToolMenu* Menu, AActor* ContextActor)
-	{
-		ILevelInstanceInterface* LevelInstanceEdit = nullptr;
-		if (ULevelInstanceSubsystem* LevelInstanceSubsystem = ContextActor->GetWorld()->GetSubsystem<ULevelInstanceSubsystem>())
-		{
-			LevelInstanceEdit = LevelInstanceSubsystem->GetEditingLevelInstance();
-
-			if (LevelInstanceEdit && !LevelInstanceSubsystem->GetEditingPropertyOverridesLevelInstance())
-			{
-				FToolUIAction LevelInstanceSetCurrentAction;
-				LevelInstanceSetCurrentAction.ExecuteAction.BindLambda([LevelInstanceEdit](const FToolMenuContext&)
-				{
-					LevelInstanceEdit->SetCurrent();
-				});
-
-				FToolMenuSection& Section = CreateLevelSection(Menu);
-				Section.AddMenuEntry(TEXT("LevelInstanceSetCurrent"), LOCTEXT("LevelInstanceSetCurrent", "Set Current Level"), TAttribute<FText>(), FSlateIcon(), LevelInstanceSetCurrentAction);
-			}
-		}
-	}
-
-	static void CreateMoveSelectionToMenu(UToolMenu* Menu, const TArray<AActor*>& ActorsToMove)
-	{
-		if (ActorsToMove.Num() > 0)
-		{
-			ILevelInstanceInterface* LevelInstanceEdit = nullptr;
-			ULevelInstanceSubsystem* LevelInstanceSubsystem = GEditor->GetEditorWorldContext().World()->GetSubsystem<ULevelInstanceSubsystem>();
-			if (LevelInstanceSubsystem)
-			{
-				LevelInstanceEdit = LevelInstanceSubsystem->GetEditingLevelInstance();
-			}
-			
-			if (LevelInstanceEdit && !LevelInstanceSubsystem->GetEditingPropertyOverridesLevelInstance())
-			{
-				FToolUIAction LevelInstanceMoveSelectionAction;
-
-				LevelInstanceMoveSelectionAction.CanExecuteAction.BindLambda([CopyActorsToMove = ActorsToMove, LevelInstanceEdit, LevelInstanceSubsystem](const FToolMenuContext& MenuContext)
-				{
-					for (AActor* ActorToMove : CopyActorsToMove)
-					{
-						if (ActorToMove->GetLevel() == LevelInstanceSubsystem->GetLevelInstanceLevel(LevelInstanceEdit))
-						{
-							return false;
-						}
-					}
-
-					return !CopyActorsToMove.IsEmpty();
-				});
-
-				LevelInstanceMoveSelectionAction.ExecuteAction.BindLambda([CopyActorsToMove = ActorsToMove, LevelInstanceEdit](const FToolMenuContext&)
-				{
-					MoveSelectionToLevelInstance(LevelInstanceEdit, CopyActorsToMove);
-				});
-
-				FToolMenuSection& Section = CreateLevelSection(Menu);
-				Section.AddMenuEntry(TEXT("LevelInstanceMoveSelectionTo"), LOCTEXT("LevelInstanceMoveSelectionTo", "Move Selection to"), TAttribute<FText>(), FSlateIcon(), LevelInstanceMoveSelectionAction);
-
+				Section.AddMenuEntry(NAME_None, LOCTEXT("LevelInstanceCancelLabel", "Cancel"), DiscardTooltip, FSlateIcon(), DiscardAction);
 			}
 		}
 	}
@@ -482,8 +432,9 @@ struct FLevelInstanceMenuUtils
 				LOCTEXT("OrganizeActorsInFolders", "Keep Folders"),
 				LOCTEXT(
 					"OrganizeActorsInFoldersTooltip",
-					"Should the actors be placed in the folder the Level Instance is in, "
-					"and keep the folder structure they had inside the Level Instance?"
+					"When checked, actors remain in the same folder as the level instance "
+					"and use the same folder structure."
+					"\nWhen unchecked, actors are placed at the root of the current level's hierarchy."
 				),
 				FSlateIcon(),
 				FUIAction(
@@ -510,9 +461,10 @@ struct FLevelInstanceMenuUtils
 					.MinValue(1)
 					.Value_Lambda([]() { return BreakLevels; })
 					.OnValueChanged_Lambda([](int32 InValue) { BreakLevels = InValue; })
+					.ToolTipText(LOCTEXT("BreakLevelsTooltip", "Determines the depth of nested instances to break apart. Use 1 to break only the top level instance."))
 					.Label()
 					[
-						SNumericEntryBox<int32>::BuildLabel(LOCTEXT("BreakLevelsLabel", "Levels"), FLinearColor::White, FLinearColor::Transparent)
+						SNumericEntryBox<int32>::BuildLabel(LOCTEXT("BreakDepthLabel", "Depth"), FLinearColor::White, FLinearColor::Transparent)
 					]
 				];
 
@@ -523,7 +475,7 @@ struct FLevelInstanceMenuUtils
 			FToolMenuEntry ExecuteEntry = FToolMenuEntry::InitMenuEntry(
 				"ExecuteBreak",
 				LOCTEXT("BreakLevelInstances_BreakLevelInstanceButton", "Break Level Instance(s)"),
-				FText(),
+				LOCTEXT("BreakLevelInstances_BreakLevelInstanceButtonTooltip", "Break apart the selected level instances using the settings above."),
 				FSlateIcon(),
 				FUIAction(
 					FExecuteAction::CreateLambda([CopyBreakableLevelInstances = BreakableLevelInstances, LevelInstanceSubsystem, Settings]()
@@ -536,7 +488,7 @@ struct FLevelInstanceMenuUtils
 							FText::AsNumber(CopyBreakableLevelInstances.Num())
 						);
 
-						if (FMessageDialog::Open(EAppMsgType::YesNo, LevelInstanceBreakWarning) == EAppReturnType::Yes)
+						if (FMessageDialog::Open(EAppMsgType::YesNo, LevelInstanceBreakWarning, LOCTEXT("BreakingLevelInstanceTitle", "Break Level Instances")) == EAppReturnType::Yes)
 						{
 							ELevelInstanceBreakFlags Flags = ELevelInstanceBreakFlags::None;
 							if (Settings->bKeepFoldersDuringBreak)
@@ -576,8 +528,8 @@ struct FLevelInstanceMenuUtils
 				FToolMenuSection& Section = CreateLevelSection(Menu);
 				Section.AddSubMenu(
 					"BreakLevelInstances",
-					LOCTEXT("BreakLevelInstances", "Break..."),
-					TAttribute<FText>(),
+					LOCTEXT("BreakLevelInstances", "Break"),
+					LOCTEXT("BreakLevelInstancesTooltip", "Break apart the selected level instances into their individual actors."),
 					FNewToolMenuDelegate::CreateLambda([CopyOfBreakableLevelInstances = BreakableLevelInstances](UToolMenu* Menu)
 					{
 						CreateBreakSubMenu(Menu, CopyOfBreakableLevelInstances);
@@ -669,7 +621,7 @@ struct FLevelInstanceMenuUtils
 
 					Section.AddMenuEntry(
 						"ResetLevelInstancePropertyOverrides",
-						LOCTEXT("ResetLevelInstancePropertyOverrides", "Reset Override(s)"),
+						LOCTEXT("ResetLevelInstancePropertyOverrides", "Reset Overrides"),
 						TAttribute<FText>(),
 						TAttribute<FSlateIcon>(),
 						UIAction);
@@ -705,8 +657,8 @@ struct FLevelInstanceMenuUtils
 
 					Section.AddMenuEntry(
 						"ResetLevelInstancePropertyOverrides",
-						LOCTEXT("ResetLevelInstancePropertyOverrides", "Reset Override(s)"),
-						TAttribute<FText>(),
+						LOCTEXT("ResetLevelInstancePropertyOverrides", "Reset Overrides"),
+						LOCTEXT("ResetLevelInstancePropertyOverridesTooltip", "Discard all overrides on the selected level instances, restoring them to match the level assets."),
 						TAttribute<FSlateIcon>(),
 						UIAction);
 				}
@@ -1093,7 +1045,7 @@ void FLevelInstanceEditorModule::RegisterLevelInstanceColumn()
 
 		FSceneOutlinerColumnInfo ColumnInfo(ESceneOutlinerColumnVisibility::Invisible, 8,
 			FCreateSceneOutlinerColumn::CreateRaw(this, &FLevelInstanceEditorModule::CreateLevelInstanceColumn),
-			true, TOptional<float>(), LOCTEXT("LevelInstanceColumnName", "Level Instance"));
+			true, TOptional<float>(), LOCTEXT("LevelInstanceColumnName", "Level Instance Overrides"));
 
 		SceneOutlinerModule.RegisterDefaultColumnType<FLevelInstanceSceneOutlinerColumn>(ColumnInfo);
 	}
@@ -1227,19 +1179,13 @@ void FLevelInstanceEditorModule::ExtendContextMenu()
 				// Allow Edit/Commmit on non root selected Level Instance
 				FLevelInstanceMenuUtils::CreateEditMenu(ToolMenu, ContextActor);
 				FLevelInstanceMenuUtils::CreateEditPropertyOverridesMenu(ToolMenu, ContextActor);
-				FLevelInstanceMenuUtils::CreateCommitDiscardMenu(ToolMenu, ContextActor);
+				FLevelInstanceMenuUtils::CreateSaveCancelMenu(ToolMenu, ContextActor);
 				
 				if (bAreAllSelectedLevelInstancesRootSelections)
 				{
 					FLevelInstanceMenuUtils::CreatePackedBlueprintMenu(ToolMenu, ContextActor);
-					FLevelInstanceMenuUtils::CreateSetCurrentMenu(ToolMenu, ContextActor);
 				}
 			}
-
-			if (bAreAllSelectedLevelInstancesRootSelections)
-			{
-				FLevelInstanceMenuUtils::CreateMoveSelectionToMenu(ToolMenu, SelectedActors);
-			}
 		}
 
 		if (bAreAllSelectedLevelInstancesRootSelections)
diff --git a/Engine/Source/Editor/LevelInstanceEditor/Private/LevelInstanceSceneOutlinerColumn.cpp b/Engine/Source/Editor/LevelInstanceEditor/Private/LevelInstanceSceneOutlinerColumn.cpp
--- a/Engine/Source/Editor/LevelInstanceEditor/Private/LevelInstanceSceneOutlinerColumn.cpp
+++ b/Engine/Source/Editor/LevelInstanceEditor/Private/LevelInstanceSceneOutlinerColumn.cpp
@@ -17,9 +17,9 @@ namespace LevelInstanceColumnPrivate
 {
 	FName Name("Level Instance");
 
-	const FText ToolTipIsOverriden = LOCTEXT("IsOverridenTooltip", "Overriden");
-	const FText ToolTipIsOverridenAndContainsOverrides = LOCTEXT("IsOverridenAndContainsOverridesTooltip", "Overriden and overrides actor(s)");
-	const FText ToolTipContainsOverrides = LOCTEXT("ContainsOverridesTooltip", "Overrides actor(s)");
+	const FText ToolTipIsOverriden = LOCTEXT("IsOverridenTooltip", "This actor is overridden.");
+	const FText ToolTipIsOverridenAndContainsOverrides = LOCTEXT("IsOverridenAndContainsOverridesTooltip", "This level instance is overridden, and so is at least one of its children.");
+	const FText ToolTipContainsOverrides = LOCTEXT("ContainsOverridesTooltip", "At least one child of this level instance is overridden.");
 
 	void GetBrushesAndToolTipForItem(FSceneOutlinerTreeItemRef TreeItem, const FSlateBrush*& OutHasOVerrideBrush, const FSlateBrush*& OutContainsOverrideBrush, const FText*& OutToolTipText)
 	{
diff --git a/Engine/Source/Runtime/Engine/Private/LevelInstance/LevelInstanceSubsystem.cpp b/Engine/Source/Runtime/Engine/Private/LevelInstance/LevelInstanceSubsystem.cpp
--- a/Engine/Source/Runtime/Engine/Private/LevelInstance/LevelInstanceSubsystem.cpp
+++ b/Engine/Source/Runtime/Engine/Private/LevelInstance/LevelInstanceSubsystem.cpp
@@ -901,49 +901,6 @@ void ULevelInstanceSubsystem::OnTryExitEditorMode()
 	}
 }
 
-FText ULevelInstanceSubsystem::GetToolKitExitToolTip() const
-{
-	if (ILevelInstanceInterface* PropertyOverrideInstance = GetEditingPropertyOverridesLevelInstance())
-	{
-		AActor* Actor = CastChecked<AActor>(PropertyOverrideInstance);
-		return FText::FromString(Actor->GetActorLabel());
-	}
-	else if (ILevelInstanceInterface* EditingLevelInstance = GetEditingLevelInstance())
-	{
-		const TSoftObjectPtr<UWorld> WorldAsset = EditingLevelInstance->GetWorldAsset();
-		return FText::FromString(WorldAsset.GetAssetName());
-	}
-
-	return FText::GetEmpty();
-}
-
-FText ULevelInstanceSubsystem::GetToolKitDisplayText() const
-{
-	if (ILevelInstanceInterface* PropertyOverrideInstance = GetEditingPropertyOverridesLevelInstance())
-	{
-		return LOCTEXT("ExitEditTooltip", "Exit Property Override Edit");
-	}
-	else if (ILevelInstanceInterface* EditingLevelInstance = GetEditingLevelInstance())
-	{
-		return LOCTEXT("ExitOverrideTooltip", "Exit Level Instance Edit");
-	}
-
-	return FText::GetEmpty();
-}
-
-void ULevelInstanceSubsystem::ToolKitExit()
-{
-	const bool bForceExit = false;
-	if (PropertyOverrideEdit)
-	{
-		TryCommitLevelInstancePropertyOverrideEdit(bForceExit);
-	}
-	else if (LevelInstanceEdit)
-	{
-		TryCommitLevelInstanceEdit(bForceExit);
-	}
-}
-
 bool ULevelInstanceSubsystem::TryCommitLevelInstanceEdit(bool bForceExit)
 {
 	if (LevelInstanceEdit)
@@ -2222,7 +2179,7 @@ bool ULevelInstanceSubsystem::CanEditLevelInstanceCommon(const ILevelInstanceInt
 	{
 		if (OutReason)
 		{
-			*OutReason = FText::Format(LOCTEXT("CanEditLevelInstanceAlreadyBeingEdited", "Level Instance already being edited ({0})."), FText::FromString(LevelInstance->GetWorldAssetPackage()));
+			*OutReason = FText::Format(LOCTEXT("CanEditLevelInstanceAlreadyBeingEdited", "This level instance is already being edited.\n\nAsset path: {0}"), FText::FromString(LevelInstance->GetWorldAssetPackage()));
 		}
 		return false;
 	}
@@ -2231,7 +2188,7 @@ bool ULevelInstanceSubsystem::CanEditLevelInstanceCommon(const ILevelInstanceInt
 	{
 		if (OutReason)
 		{
-			*OutReason = FText::Format(LOCTEXT("CanEditLevelInstanceAlreadyBeingOverriden", "Level Instance already in property override edit ({0})."), FText::FromString(LevelInstance->GetWorldAssetPackage()));
+			*OutReason = FText::Format(LOCTEXT("CanEditLevelInstanceAlreadyBeingOverriden", "This level instance is already being overridden.\n\nAsset path: {0}"), FText::FromString(LevelInstance->GetWorldAssetPackage()));
 		}
 		return false;
 	}
@@ -2342,7 +2299,7 @@ bool ULevelInstanceSubsystem::CanEditLevelInstancePropertyOverrides(const ILevel
 	{
 		if (OutReason)
 		{
-			*OutReason = LOCTEXT("LevelInstancePropertyOverridesWorldPartitionOnly", "Property overrides are only supported on world partition level hierarchies");
+			*OutReason = LOCTEXT("LevelInstancePropertyOverridesWorldPartitionOnly", "Overrides are only supported for levels that use World Partition.");
 		}
 		return false;
 	}
@@ -2351,7 +2308,7 @@ bool ULevelInstanceSubsystem::CanEditLevelInstancePropertyOverrides(const ILevel
 	{
 		if (OutReason)
 		{
-			*OutReason = LOCTEXT("LevelInstancePropertyOverridesNewlyCreated", "Property overrides are only supported on saved worlds. Please save world first.");
+			*OutReason = LOCTEXT("LevelInstancePropertyOverridesNewlyCreated", "Overrides are only supported for saved levels. Save the level first.");
 		}
 		return false;
 	}
@@ -3063,6 +3020,7 @@ bool ULevelInstanceSubsystem::CommitLevelInstancePropertyOverridesInternal(TUniq
 	ILevelInstanceInterface* LevelInstance = InPropertyOverrideEdit->GetLevelInstance();
 	ILevelInstanceInterface* LevelInstanceWithOverrides = GetLevelInstancePropertyOverridesEditOwner(LevelInstance);
 	TSoftObjectPtr<ULevelInstancePropertyOverrideAsset> PreviousOverrideAsset = LevelInstanceWithOverrides->GetPropertyOverrideAsset() ? LevelInstanceWithOverrides->GetPropertyOverrideAsset()->GetSourceAssetPtr() : nullptr;
+	const FLevelInstanceID LevelInstanceWithOverridesID = LevelInstanceWithOverrides->GetLevelInstanceID();
 
 	bool bSaved = false;
 	if (InPropertyOverrideEdit->IsDirty() && !bDiscardEdits)
@@ -3088,7 +3046,11 @@ bool ULevelInstanceSubsystem::CommitLevelInstancePropertyOverridesInternal(TUniq
 	}
 
 	BlockOnLoading();
-			
+
+	if (AActor* LevelInstanceToSelect = Cast<AActor>(GetLevelInstance(LevelInstanceWithOverridesID)))
+	{
+		GEditor->SelectActor(LevelInstanceToSelect, true, true);
+	}
 	return true;
 }
 
@@ -3117,6 +3079,7 @@ void ULevelInstanceSubsystem::ResetPropertyOverrides(ILevelInstanceInterface* Le
 		return;
 	}
 
+	const FLevelInstanceID LevelInstanceWithOverridesID = LevelInstance->GetLevelInstanceID();
 	const TSoftObjectPtr<ULevelInstancePropertyOverrideAsset> AssetPath = LevelInstance->GetPropertyOverrideAsset()->GetSourceAssetPtr();
 	ULevelInstancePropertyOverrideAsset* PropertyOverrideAsset = LevelInstance->GetPropertyOverrideAsset();
 	
@@ -3132,12 +3095,17 @@ void ULevelInstanceSubsystem::ResetPropertyOverrides(ILevelInstanceInterface* Le
 		return;
 	}
 
+	GEditor->SelectNone(true, true);
+
 	LevelInstance->UpdateLevelInstanceFromWorldAsset();
 	UpdateLevelInstancesFromPropertyOverrideAsset(AssetPath, nullptr);
 	BlockOnLoading();
 	
 	// Make sure selection is refreshed (Reset can have impact on details view)
-	GEditor->SelectNone(true, true);
+	if (AActor* LevelInstanceToSelect = Cast<AActor>(GetLevelInstance(LevelInstanceWithOverridesID)))
+	{
+		GEditor->SelectActor(LevelInstanceToSelect, true, true);
+	}
 }
 
 void ULevelInstanceSubsystem::RegisterPrimitiveColorHandler()
diff --git a/Engine/Source/Runtime/Engine/Public/LevelInstance/LevelInstanceSubsystem.h b/Engine/Source/Runtime/Engine/Public/LevelInstance/LevelInstanceSubsystem.h
--- a/Engine/Source/Runtime/Engine/Public/LevelInstance/LevelInstanceSubsystem.h
+++ b/Engine/Source/Runtime/Engine/Public/LevelInstance/LevelInstanceSubsystem.h
@@ -175,9 +175,14 @@ public:
 	ENGINE_API bool IsSubSelectionEnabled() const;
 private:
 	friend class FLevelInstanceEditorModule;
+	friend struct FLevelInstanceActorDetailsHelper;
+	friend struct FLevelInstanceEditorModeToolkitHelper;
 	friend struct FLevelInstanceMenuUtils;
 	friend class ILevelInstanceInterface;
 	friend class ULevelInstanceEditorMode;
+	friend class ULevelStreamingLevelInstanceEditor;
+	friend class ULevelStreamingLevelInstanceEditorPropertyOverride;
+	friend class ULevelInstancePropertyOverrideAsset;
 
 	// EditPropertyOverrides
 	ENGINE_API TArray<ILevelInstanceInterface*> GetLevelInstances(const TSoftObjectPtr<ULevelInstancePropertyOverrideAsset>& PropertyOverrideAsset) const;
@@ -195,6 +200,7 @@ private:
 	ENGINE_API static void RegisterPrimitiveColorHandler();
 	ENGINE_API static void UnregisterPrimitiveColorHandler();
 #endif
+	friend class ULevelStreamingLevelInstance;
 
 private:
 	ENGINE_API void UpdateStreamingStateInternal();
@@ -272,10 +278,6 @@ private:
 	ENGINE_API bool CommitLevelInstancePropertyOverridesInternal(TUniquePtr<FPropertyOverrideEdit>& InPropertyOverrideEdit, bool bDiscardEdits);
 	bool CanEditLevelInstanceCommon(const ILevelInstanceInterface* LevelInstance, FText* OutReason = nullptr) const;
 
-	ENGINE_API FText GetToolKitExitToolTip() const;
-	ENGINE_API FText GetToolKitDisplayText() const;
-	ENGINE_API void ToolKitExit();
-
 	void OnExitEditorModeInternal(bool bForceExit);
 	bool TryCommitLevelInstanceEdit(bool bForceExit);
 	bool TryCommitLevelInstancePropertyOverrideEdit(bool bForceExit);
@@ -310,11 +312,6 @@ private:
 
 private:
 #endif
-	friend class ULevelStreamingLevelInstance;
-	friend class ULevelStreamingLevelInstanceEditor;
-	friend class ULevelStreamingLevelInstanceEditorPropertyOverride;
-	friend class ULevelInstancePropertyOverrideAsset;
-	friend class FLevelInstanceEditorModeToolkit;
 
 #if WITH_EDITOR
 	bool bIsCreatingLevelInstance;
