[Mode Manager] Fix some singleton startup pattern that was allocating at incorrect timings to speculatively fix some allocate on shutdown crash(es)

Also move some of the global access to the level editor mode manager to pipe through the level editor where there were some early accessors on startup.

diff --git a/Engine/Plugins/PCG/Source/PCGEditor/Private/PCGEditorModule.cpp b/Engine/Plugins/PCG/Source/PCGEditor/Private/PCGEditorModule.cpp
--- a/Engine/Plugins/PCG/Source/PCGEditor/Private/PCGEditorModule.cpp
+++ b/Engine/Plugins/PCG/Source/PCGEditor/Private/PCGEditorModule.cpp
@@ -43,6 +43,7 @@
 #include "Details/PCGInstancedPropertyBagOverrideDetails.h"
 #include "Details/PCGVolumeDetails.h"
 #include "Framework/MultiBox/MultiBoxBuilder.h"
+#include "LevelEditor.h"
 #include "Framework/Notifications/NotificationManager.h"
 
 #define LOCTEXT_NAMESPACE "FPCGEditorModule"
@@ -61,7 +62,15 @@ void FPCGEditorModule::StartupModule()
 	GraphNodeFactory = MakeShareable(new FPCGEditorGraphNodeFactory());
 	FEdGraphUtilities::RegisterVisualNodeFactory(GraphNodeFactory);
 
-	FCoreDelegates::OnPostEngineInit.AddRaw(this, &FPCGEditorModule::OnPostEngineInit);
+	FLevelEditorModule& LevelEditorModule = FModuleManager::LoadModuleChecked<FLevelEditorModule>("LevelEditor");
+	if (TSharedPtr<ILevelEditor> FirstLevelEditor = LevelEditorModule.GetFirstLevelEditor())
+	{
+		OnLevelEditorCreated(FirstLevelEditor);
+	}
+	else
+	{
+		LevelEditorModule.OnLevelEditorCreated().AddRaw(this, &FPCGEditorModule::OnLevelEditorCreated);
+	}
 }
 
 void FPCGEditorModule::ShutdownModule()
@@ -86,17 +95,18 @@ void FPCGEditorModule::ShutdownModule()
 		});
 	}
 
-	FCoreDelegates::OnPostEngineInit.RemoveAll(this);
-	if (!IsRunningCommandlet())
+	if (FModuleManager::Get().IsModuleLoaded("LevelEditor"))
 	{
-		if (GLevelEditorModeToolsIsValid())
+		FLevelEditorModule& LevelEditorModule = FModuleManager::GetModuleChecked<FLevelEditorModule>("LevelEditor");
+		LevelEditorModule.OnLevelEditorCreated().RemoveAll(this);
+		if (TSharedPtr<ILevelEditor> FirstLevelEditor = LevelEditorModule.GetFirstLevelEditor())
 		{
-			GLevelEditorModeTools().OnEditorModeIDChanged().RemoveAll(this);
+			FirstLevelEditor->GetEditorModeManager().OnEditorModeIDChanged().RemoveAll(this);
 		}
 	}
 }
 
-void FPCGEditorModule::OnPostEngineInit()
+void FPCGEditorModule::OnLevelEditorCreated(TSharedPtr<ILevelEditor> InLevelEditor)
 {
 	RegisterOnEditorModeChange();
 
@@ -168,9 +178,10 @@ void FPCGEditorModule::ReleaseProgressNotification(TWeakPtr<IPCGEditorProgressNo
 void FPCGEditorModule::RegisterOnEditorModeChange()
 {
 	// Have a callback that catches changes in the Editor modes, to catch when we exit the landscape edit mode.
-	if (!IsRunningCommandlet() && GLevelEditorModeToolsIsValid())
+	FLevelEditorModule& LevelEditorModule = FModuleManager::GetModuleChecked<FLevelEditorModule>("LevelEditor");
+	if (TSharedPtr<ILevelEditor> FirstLevelEditor = LevelEditorModule.GetFirstLevelEditor())
 	{
-		GLevelEditorModeTools().OnEditorModeIDChanged().AddRaw(this, &FPCGEditorModule::OnEditorModeIDChanged);
+		FirstLevelEditor->GetEditorModeManager().OnEditorModeIDChanged().AddRaw(this, &FPCGEditorModule::OnEditorModeIDChanged);
 	}
 }
 
diff --git a/Engine/Plugins/PCG/Source/PCGEditor/Private/PCGEditorModule.h b/Engine/Plugins/PCG/Source/PCGEditor/Private/PCGEditorModule.h
--- a/Engine/Plugins/PCG/Source/PCGEditor/Private/PCGEditorModule.h
+++ b/Engine/Plugins/PCG/Source/PCGEditor/Private/PCGEditorModule.h
@@ -8,6 +8,7 @@
 #include "Modules/ModuleInterface.h"
 #include "Toolkits/AssetEditorToolkit.h"
 
+class ILevelEditor;
 // Logs
 DECLARE_LOG_CATEGORY_EXTERN(LogPCGEditor, Log, All);
 
@@ -40,7 +41,7 @@ protected:
 	void RegisterPCGDataVisualizations();
 	void UnregisterPCGDataVisualizations();
 
-	void OnPostEngineInit();
+	void OnLevelEditorCreated(TSharedPtr<ILevelEditor> InLevelEditor);
 	void RegisterOnEditorModeChange();
 	void OnEditorModeIDChanged(const FEditorModeID& EditorModeID, bool bIsEntering);
 
diff --git a/Engine/Source/Editor/ActorPickerMode/Private/ActorPickerMode.cpp b/Engine/Source/Editor/ActorPickerMode/Private/ActorPickerMode.cpp
--- a/Engine/Source/Editor/ActorPickerMode/Private/ActorPickerMode.cpp
+++ b/Engine/Source/Editor/ActorPickerMode/Private/ActorPickerMode.cpp
@@ -7,11 +7,14 @@
 #include "EditorModes.h"
 #include "EditorModeActorPicker.h" 
 #include "Framework/Application/SlateApplication.h"
+#include "LevelEditor.h"
 
 IMPLEMENT_MODULE( FActorPickerModeModule, ActorPickerMode );
 
 void FActorPickerModeModule::StartupModule()
 {
+	// Ensure the level editor module is loaded for grabbing the mode manager later
+	FModuleManager::Get().LoadModuleChecked("LevelEditor");
 	FEditorModeRegistry::Get().RegisterMode<FEdModeActorPicker>(FBuiltinEditorModes::EM_ActorPicker);
 
 	if (FSlateApplication::IsInitialized())
@@ -33,21 +36,19 @@ void FActorPickerModeModule::ShutdownModule()
 
 void FActorPickerModeModule::BeginActorPickingMode(FOnGetAllowedClasses InOnGetAllowedClasses, FOnShouldFilterActor InOnShouldFilterActor, FOnActorSelected InOnActorSelected) const
 {
-	if(!GLevelEditorModeToolsIsValid())
+	if (FEditorModeTools* ModeTools = GetLevelEditorModeManager())
 	{
-		return;
-	}
-
-	// Activate the mode
-	GLevelEditorModeTools().ActivateMode(FBuiltinEditorModes::EM_ActorPicker);
+		// Activate the mode
+		ModeTools->ActivateMode(FBuiltinEditorModes::EM_ActorPicker);
 
-	// Set the required delegates
-	FEdModeActorPicker* Mode = GLevelEditorModeTools().GetActiveModeTyped<FEdModeActorPicker>(FBuiltinEditorModes::EM_ActorPicker);
-	if (ensure(Mode))
-	{
-		Mode->OnActorSelected = InOnActorSelected;
-		Mode->OnGetAllowedClasses = InOnGetAllowedClasses;
-		Mode->OnShouldFilterActor = InOnShouldFilterActor;
+		// Set the required delegates
+		FEdModeActorPicker* Mode = ModeTools->GetActiveModeTyped<FEdModeActorPicker>(FBuiltinEditorModes::EM_ActorPicker);
+		if (ensure(Mode))
+		{
+			Mode->OnActorSelected = InOnActorSelected;
+			Mode->OnGetAllowedClasses = InOnGetAllowedClasses;
+			Mode->OnShouldFilterActor = InOnShouldFilterActor;
+		}
 	}
 }
 
@@ -61,17 +62,29 @@ void FActorPickerModeModule::OnApplicationDeactivated(const bool IsActive) const
 
 void FActorPickerModeModule::EndActorPickingMode() const
 {
-	if (IsInActorPickingMode() && GLevelEditorModeToolsIsValid())
+	if (FEditorModeTools* ModeTools = GetLevelEditorModeManager())
 	{
-		GLevelEditorModeTools().DeactivateMode(FBuiltinEditorModes::EM_ActorPicker);
+		ModeTools->DeactivateMode(FBuiltinEditorModes::EM_ActorPicker);
 	}
 }
 
 bool FActorPickerModeModule::IsInActorPickingMode() const
 {
-	if(GLevelEditorModeToolsIsValid())
+	if (FEditorModeTools* ModeTools = GetLevelEditorModeManager())
 	{
-		return GLevelEditorModeTools().IsModeActive(FBuiltinEditorModes::EM_ActorPicker);
+		return ModeTools->IsModeActive(FBuiltinEditorModes::EM_ActorPicker);
 	}
 	return false;
 }
+
+FEditorModeTools* FActorPickerModeModule::GetLevelEditorModeManager()
+{
+	FLevelEditorModule& LevelEditorModule = FModuleManager::GetModuleChecked<FLevelEditorModule>("LevelEditor");
+	TSharedPtr<ILevelEditor> FirstLevelEditor = LevelEditorModule.GetFirstLevelEditor();
+	if (FirstLevelEditor.IsValid())
+	{
+		return &FirstLevelEditor->GetEditorModeManager();
+	}
+
+	return nullptr;
+}
diff --git a/Engine/Source/Editor/ActorPickerMode/Public/ActorPickerMode.h b/Engine/Source/Editor/ActorPickerMode/Public/ActorPickerMode.h
--- a/Engine/Source/Editor/ActorPickerMode/Public/ActorPickerMode.h
+++ b/Engine/Source/Editor/ActorPickerMode/Public/ActorPickerMode.h
@@ -6,6 +6,7 @@
 #include "Modules/ModuleInterface.h"
 
 class AActor;
+class FEditorModeTools;
 
 DECLARE_DELEGATE_OneParam( FOnGetAllowedClasses, TArray<const UClass*>& );
 DECLARE_DELEGATE_OneParam( FOnActorSelected, AActor* );
@@ -41,6 +42,8 @@ private:
 	/** Handler for when the application is deactivated. */
 	void OnApplicationDeactivated(const bool IsActive) const;
 
+	static FEditorModeTools* GetLevelEditorModeManager();
+
 private:
 	FDelegateHandle OnApplicationDeactivatedHandle;
 };
diff --git a/Engine/Source/Editor/LevelInstanceEditor/Private/LevelInstanceEditorModule.cpp b/Engine/Source/Editor/LevelInstanceEditor/Private/LevelInstanceEditorModule.cpp
--- a/Engine/Source/Editor/LevelInstanceEditor/Private/LevelInstanceEditorModule.cpp
+++ b/Engine/Source/Editor/LevelInstanceEditor/Private/LevelInstanceEditorModule.cpp
@@ -918,6 +918,15 @@ void FLevelInstanceEditorModule::RegisterToFirstLevelEditor()
 	if (FirstLevelEditor.IsValid())
 	{
 		FirstLevelEditor->AddActorDetailsSCSEditorUICustomization(FLevelInstanceActorDetailsSCSEditorUICustomization::GetInstance());
+
+		FEditorModeTools& LevelEditorModeManager = FirstLevelEditor->GetEditorModeManager();
+		LevelEditorModeManager.OnEditorModeIDChanged().AddRaw(this, &FLevelInstanceEditorModule::OnEditorModeIDChanged);
+
+		// Create a Behavior source for the default EdModeTools (when we aren't in the LevelInstanceEditorMode)
+		DefaultBehaviorSource = ULevelInstanceEditorMode::CreateDefaultModeBehaviorSource(LevelEditorModeManager.GetInteractiveToolsContext());
+		LevelEditorModeManager.GetInteractiveToolsContext()->InputRouter->RegisterSource(DefaultBehaviorSource.GetInterface());
+		
+		RegisterLevelInstanceColumn();
 	}
 }
 
@@ -985,17 +994,6 @@ void FLevelInstanceEditorModule::StartupModule()
 		
 	FLevelInstanceEditorModeCommands::Register();
 
-	if (!IsRunningCommandlet())
-	{
-		GLevelEditorModeTools().OnEditorModeIDChanged().AddRaw(this, &FLevelInstanceEditorModule::OnEditorModeIDChanged);
-	
-		// Create a Behavior source for the default EdModeTools (when we aren't in the LevelInstanceEditorMode)
-		DefaultBehaviorSource = ULevelInstanceEditorMode::CreateDefaultModeBehaviorSource(GLevelEditorModeTools().GetInteractiveToolsContext());
-		GLevelEditorModeTools().GetInteractiveToolsContext()->InputRouter->RegisterSource(DefaultBehaviorSource.GetInterface());
-
-		RegisterLevelInstanceColumn();
-	}
-
 	ULevelInstanceSubsystem::RegisterPrimitiveColorHandler();
 }
 
@@ -1010,7 +1008,13 @@ void FLevelInstanceEditorModule::ShutdownModule()
 		if (TSharedPtr<ILevelEditor> FirstLevelEditor = LevelEditorModule.GetFirstLevelEditor())
 		{
 			FirstLevelEditor->RemoveActorDetailsSCSEditorUICustomization(FLevelInstanceActorDetailsSCSEditorUICustomization::GetInstance());
+			FirstLevelEditor->GetEditorModeManager().OnEditorModeIDChanged().RemoveAll(this);
+			FirstLevelEditor->GetEditorModeManager().GetInteractiveToolsContext()->InputRouter->DeregisterSource(DefaultBehaviorSource.GetInterface());
 		}
+
+		DefaultBehaviorSource = nullptr;
+		
+		UnregisterLevelInstanceColumn();
 	}
 
 	if (GEditor)
@@ -1019,17 +1023,6 @@ void FLevelInstanceEditorModule::ShutdownModule()
 	}
 
 	EditorLevelUtils::CanMoveActorToLevelDelegate.RemoveAll(this);
-
-	if (!IsRunningCommandlet())
-	{
-		if (GLevelEditorModeToolsIsValid())
-		{
-			GLevelEditorModeTools().OnEditorModeIDChanged().RemoveAll(this);
-			GLevelEditorModeTools().GetInteractiveToolsContext()->InputRouter->DeregisterSource(DefaultBehaviorSource.GetInterface());
-			DefaultBehaviorSource = nullptr;
-		}
-		UnregisterLevelInstanceColumn();
-	}
 }
 
 TSharedRef<ISceneOutlinerColumn> FLevelInstanceEditorModule::CreateLevelInstanceColumn(ISceneOutliner& SceneOutliner) const
@@ -1074,15 +1067,16 @@ void FLevelInstanceEditorModule::BroadcastTryExitEditorMode()
 
 void FLevelInstanceEditorModule::UpdateEditorMode(bool bActivated)
 {
-	if (!IsRunningCommandlet() && GLevelEditorModeToolsIsValid())
+	FLevelEditorModule& LevelEditorModule = FModuleManager::GetModuleChecked<FLevelEditorModule>("LevelEditor");
+	if (TSharedPtr<ILevelEditor> FirstLevelEditor = LevelEditorModule.GetFirstLevelEditor())
 	{
-		if (bActivated && !GLevelEditorModeTools().IsModeActive(ULevelInstanceEditorMode::EM_LevelInstanceEditorModeId))
+		if (bActivated && !FirstLevelEditor->GetEditorModeManager().IsModeActive(ULevelInstanceEditorMode::EM_LevelInstanceEditorModeId))
 		{
-			GLevelEditorModeTools().ActivateMode(ULevelInstanceEditorMode::EM_LevelInstanceEditorModeId);
+			FirstLevelEditor->GetEditorModeManager().ActivateMode(ULevelInstanceEditorMode::EM_LevelInstanceEditorModeId);
 		}
-		else if (!bActivated && GLevelEditorModeTools().IsModeActive(ULevelInstanceEditorMode::EM_LevelInstanceEditorModeId))
+		else if (!bActivated && FirstLevelEditor->GetEditorModeManager().IsModeActive(ULevelInstanceEditorMode::EM_LevelInstanceEditorModeId))
 		{
-			GLevelEditorModeTools().DeactivateMode(ULevelInstanceEditorMode::EM_LevelInstanceEditorModeId);
+			FirstLevelEditor->GetEditorModeManager().DeactivateMode(ULevelInstanceEditorMode::EM_LevelInstanceEditorModeId);
 		}
 	}
 }
diff --git a/Engine/Source/Editor/UnrealEd/Private/UnrealEdGlobals.cpp b/Engine/Source/Editor/UnrealEd/Private/UnrealEdGlobals.cpp
--- a/Engine/Source/Editor/UnrealEd/Private/UnrealEdGlobals.cpp
+++ b/Engine/Source/Editor/UnrealEd/Private/UnrealEdGlobals.cpp
@@ -47,24 +47,22 @@ DEFINE_LOG_CATEGORY_STATIC(LogUnrealEd, Log, All);
  */
 namespace Internal
 {
-	static TSharedPtr<FEditorModeTools>& GetGlobalModeManager()
-	{
-		checkf(!IsRunningCommandlet(), TEXT("The global mode manager should not be created or accessed in a commandlet environment. Check that your mode or module is not accessing the global mode tools or that scriptable features of modes have been moved to subsystems."));
-
-		static TSharedPtr<FEditorModeTools> EditorModeToolsSingleton = MakeShared<FEditorModeTools>();
-		return EditorModeToolsSingleton;
-	}
+	static TSharedPtr<FEditorModeTools> EditorModeToolsSingleton;
 };
 
 bool GLevelEditorModeToolsIsValid()
 {
-	return Internal::GetGlobalModeManager().IsValid();
+	return Internal::EditorModeToolsSingleton.IsValid();
 }
 
 FEditorModeTools& GLevelEditorModeTools()
 {
-	check(GLevelEditorModeToolsIsValid());
-	return *Internal::GetGlobalModeManager().Get();
+	checkf(!IsRunningCommandlet(), TEXT("The global mode manager should not be created or accessed in a commandlet environment. Check that your mode or module is not accessing the global mode tools or that scriptable features of modes have been moved to subsystems."));
+	if (!ensureMsgf(Internal::EditorModeToolsSingleton.IsValid(), TEXT("The level editor is not started up yet. If you need to access the global mode manager early in the startup phase, please use FLevelEditorModule::OnLevelEditorCreated to gate the access.")))
+	{
+		Internal::EditorModeToolsSingleton = MakeShared<FEditorModeTools>();
+	}
+	return *Internal::EditorModeToolsSingleton.Get();
 }
 
 FLevelEditorViewportClient* GCurrentLevelEditingViewportClient = NULL;
@@ -129,16 +127,20 @@ int32 EditorInit( IEngineLoop& EngineLoop )
 
 	// Set up the actor folders singleton
 	FActorFolders::Get();
+	if (!Internal::EditorModeToolsSingleton.IsValid())
+	{
+		Internal::EditorModeToolsSingleton = MakeShared<FEditorModeTools>();
+	}
 
 	// Initialize the misc editor
 	FUnrealEdMisc::Get().OnInit();
 	FCoreDelegates::OnExit.AddLambda([]()
 	{
 		// Shutdown the global static mode manager
-		if (Internal::GetGlobalModeManager().IsValid())
+		if (Internal::EditorModeToolsSingleton.IsValid())
 		{
 			GLevelEditorModeTools().SetDefaultMode(FBuiltinEditorModes::EM_Default);
-			Internal::GetGlobalModeManager().Reset();
+			Internal::EditorModeToolsSingleton.Reset();
 		}
 	});
 	
diff --git a/Engine/Source/Editor/UnrealEd/Public/Editor.h b/Engine/Source/Editor/UnrealEd/Public/Editor.h
--- a/Engine/Source/Editor/UnrealEd/Public/Editor.h
+++ b/Engine/Source/Editor/UnrealEd/Public/Editor.h
@@ -713,6 +713,7 @@ UNREALED_API class FEditorModeTools& GLevelEditorModeTools();
 /**
  * Checks if FEditorModeTools is valid
  */
+UE_DEPRECATED(5.5, "Checking the validity of the global mode manager is unnecessary. Instead use FLevelEditorModule::OnLevelEditorCreated to gate the access on the global mode manager.")
 UNREALED_API bool GLevelEditorModeToolsIsValid();
 
 namespace EditorUtilities
