[Backout] - CL33722326
Original CL Desc
-----------------------------------------------------------------
[Mode Manager] Fix some singleton startup pattern that was allocating at incorrect timings to speculatively fix some allocate on shutdown crash(es)

Also move some of the global access to the level editor mode manager to pipe through the level editor where there were some early accessors on startup.

diff --git a/Engine/Plugins/PCG/Source/PCGEditor/Private/PCGEditorModule.cpp b/Engine/Plugins/PCG/Source/PCGEditor/Private/PCGEditorModule.cpp
--- a/Engine/Plugins/PCG/Source/PCGEditor/Private/PCGEditorModule.cpp
+++ b/Engine/Plugins/PCG/Source/PCGEditor/Private/PCGEditorModule.cpp
@@ -43,7 +43,6 @@
 #include "Details/PCGInstancedPropertyBagOverrideDetails.h"
 #include "Details/PCGVolumeDetails.h"
 #include "Framework/MultiBox/MultiBoxBuilder.h"
-#include "LevelEditor.h"
 #include "Framework/Notifications/NotificationManager.h"
 
 #define LOCTEXT_NAMESPACE "FPCGEditorModule"
@@ -62,15 +61,7 @@ void FPCGEditorModule::StartupModule()
 	GraphNodeFactory = MakeShareable(new FPCGEditorGraphNodeFactory());
 	FEdGraphUtilities::RegisterVisualNodeFactory(GraphNodeFactory);
 
-	FLevelEditorModule& LevelEditorModule = FModuleManager::LoadModuleChecked<FLevelEditorModule>("LevelEditor");
-	if (TSharedPtr<ILevelEditor> FirstLevelEditor = LevelEditorModule.GetFirstLevelEditor())
-	{
-		OnLevelEditorCreated(FirstLevelEditor);
-	}
-	else
-	{
-		LevelEditorModule.OnLevelEditorCreated().AddRaw(this, &FPCGEditorModule::OnLevelEditorCreated);
-	}
+	FCoreDelegates::OnPostEngineInit.AddRaw(this, &FPCGEditorModule::OnPostEngineInit);
 }
 
 void FPCGEditorModule::ShutdownModule()
@@ -95,18 +86,17 @@ void FPCGEditorModule::ShutdownModule()
 		});
 	}
 
-	if (FModuleManager::Get().IsModuleLoaded("LevelEditor"))
+	FCoreDelegates::OnPostEngineInit.RemoveAll(this);
+	if (!IsRunningCommandlet())
 	{
-		FLevelEditorModule& LevelEditorModule = FModuleManager::GetModuleChecked<FLevelEditorModule>("LevelEditor");
-		LevelEditorModule.OnLevelEditorCreated().RemoveAll(this);
-		if (TSharedPtr<ILevelEditor> FirstLevelEditor = LevelEditorModule.GetFirstLevelEditor())
+		if (GLevelEditorModeToolsIsValid())
 		{
-			FirstLevelEditor->GetEditorModeManager().OnEditorModeIDChanged().RemoveAll(this);
+			GLevelEditorModeTools().OnEditorModeIDChanged().RemoveAll(this);
 		}
 	}
 }
 
-void FPCGEditorModule::OnLevelEditorCreated(TSharedPtr<ILevelEditor> InLevelEditor)
+void FPCGEditorModule::OnPostEngineInit()
 {
 	RegisterOnEditorModeChange();
 
@@ -178,10 +168,9 @@ void FPCGEditorModule::ReleaseProgressNotification(TWeakPtr<IPCGEditorProgressNo
 void FPCGEditorModule::RegisterOnEditorModeChange()
 {
 	// Have a callback that catches changes in the Editor modes, to catch when we exit the landscape edit mode.
-	FLevelEditorModule& LevelEditorModule = FModuleManager::GetModuleChecked<FLevelEditorModule>("LevelEditor");
-	if (TSharedPtr<ILevelEditor> FirstLevelEditor = LevelEditorModule.GetFirstLevelEditor())
+	if (!IsRunningCommandlet() && GLevelEditorModeToolsIsValid())
 	{
-		FirstLevelEditor->GetEditorModeManager().OnEditorModeIDChanged().AddRaw(this, &FPCGEditorModule::OnEditorModeIDChanged);
+		GLevelEditorModeTools().OnEditorModeIDChanged().AddRaw(this, &FPCGEditorModule::OnEditorModeIDChanged);
 	}
 }
 
diff --git a/Engine/Plugins/PCG/Source/PCGEditor/Private/PCGEditorModule.h b/Engine/Plugins/PCG/Source/PCGEditor/Private/PCGEditorModule.h
--- a/Engine/Plugins/PCG/Source/PCGEditor/Private/PCGEditorModule.h
+++ b/Engine/Plugins/PCG/Source/PCGEditor/Private/PCGEditorModule.h
@@ -8,7 +8,6 @@
 #include "Modules/ModuleInterface.h"
 #include "Toolkits/AssetEditorToolkit.h"
 
-class ILevelEditor;
 // Logs
 DECLARE_LOG_CATEGORY_EXTERN(LogPCGEditor, Log, All);
 
@@ -41,7 +40,7 @@ protected:
 	void RegisterPCGDataVisualizations();
 	void UnregisterPCGDataVisualizations();
 
-	void OnLevelEditorCreated(TSharedPtr<ILevelEditor> InLevelEditor);
+	void OnPostEngineInit();
 	void RegisterOnEditorModeChange();
 	void OnEditorModeIDChanged(const FEditorModeID& EditorModeID, bool bIsEntering);
 
diff --git a/Engine/Source/Editor/ActorPickerMode/Private/ActorPickerMode.cpp b/Engine/Source/Editor/ActorPickerMode/Private/ActorPickerMode.cpp
--- a/Engine/Source/Editor/ActorPickerMode/Private/ActorPickerMode.cpp
+++ b/Engine/Source/Editor/ActorPickerMode/Private/ActorPickerMode.cpp
@@ -7,14 +7,11 @@
 #include "EditorModes.h"
 #include "EditorModeActorPicker.h" 
 #include "Framework/Application/SlateApplication.h"
-#include "LevelEditor.h"
 
 IMPLEMENT_MODULE( FActorPickerModeModule, ActorPickerMode );
 
 void FActorPickerModeModule::StartupModule()
 {
-	// Ensure the level editor module is loaded for grabbing the mode manager later
-	FModuleManager::Get().LoadModuleChecked("LevelEditor");
 	FEditorModeRegistry::Get().RegisterMode<FEdModeActorPicker>(FBuiltinEditorModes::EM_ActorPicker);
 
 	if (FSlateApplication::IsInitialized())
@@ -36,19 +33,21 @@ void FActorPickerModeModule::ShutdownModule()
 
 void FActorPickerModeModule::BeginActorPickingMode(FOnGetAllowedClasses InOnGetAllowedClasses, FOnShouldFilterActor InOnShouldFilterActor, FOnActorSelected InOnActorSelected) const
 {
-	if (FEditorModeTools* ModeTools = GetLevelEditorModeManager())
+	if(!GLevelEditorModeToolsIsValid())
 	{
-		// Activate the mode
-		ModeTools->ActivateMode(FBuiltinEditorModes::EM_ActorPicker);
+		return;
+	}
+
+	// Activate the mode
+	GLevelEditorModeTools().ActivateMode(FBuiltinEditorModes::EM_ActorPicker);
 
-		// Set the required delegates
-		FEdModeActorPicker* Mode = ModeTools->GetActiveModeTyped<FEdModeActorPicker>(FBuiltinEditorModes::EM_ActorPicker);
-		if (ensure(Mode))
-		{
-			Mode->OnActorSelected = InOnActorSelected;
-			Mode->OnGetAllowedClasses = InOnGetAllowedClasses;
-			Mode->OnShouldFilterActor = InOnShouldFilterActor;
-		}
+	// Set the required delegates
+	FEdModeActorPicker* Mode = GLevelEditorModeTools().GetActiveModeTyped<FEdModeActorPicker>(FBuiltinEditorModes::EM_ActorPicker);
+	if (ensure(Mode))
+	{
+		Mode->OnActorSelected = InOnActorSelected;
+		Mode->OnGetAllowedClasses = InOnGetAllowedClasses;
+		Mode->OnShouldFilterActor = InOnShouldFilterActor;
 	}
 }
 
@@ -62,29 +61,17 @@ void FActorPickerModeModule::OnApplicationDeactivated(const bool IsActive) const
 
 void FActorPickerModeModule::EndActorPickingMode() const
 {
-	if (FEditorModeTools* ModeTools = GetLevelEditorModeManager())
+	if (IsInActorPickingMode() && GLevelEditorModeToolsIsValid())
 	{
-		ModeTools->DeactivateMode(FBuiltinEditorModes::EM_ActorPicker);
+		GLevelEditorModeTools().DeactivateMode(FBuiltinEditorModes::EM_ActorPicker);
 	}
 }
 
 bool FActorPickerModeModule::IsInActorPickingMode() const
 {
-	if (FEditorModeTools* ModeTools = GetLevelEditorModeManager())
+	if(GLevelEditorModeToolsIsValid())
 	{
-		return ModeTools->IsModeActive(FBuiltinEditorModes::EM_ActorPicker);
+		return GLevelEditorModeTools().IsModeActive(FBuiltinEditorModes::EM_ActorPicker);
 	}
 	return false;
 }
-
-FEditorModeTools* FActorPickerModeModule::GetLevelEditorModeManager()
-{
-	FLevelEditorModule& LevelEditorModule = FModuleManager::GetModuleChecked<FLevelEditorModule>("LevelEditor");
-	TSharedPtr<ILevelEditor> FirstLevelEditor = LevelEditorModule.GetFirstLevelEditor();
-	if (FirstLevelEditor.IsValid())
-	{
-		return &FirstLevelEditor->GetEditorModeManager();
-	}
-
-	return nullptr;
-}
diff --git a/Engine/Source/Editor/ActorPickerMode/Public/ActorPickerMode.h b/Engine/Source/Editor/ActorPickerMode/Public/ActorPickerMode.h
--- a/Engine/Source/Editor/ActorPickerMode/Public/ActorPickerMode.h
+++ b/Engine/Source/Editor/ActorPickerMode/Public/ActorPickerMode.h
@@ -6,7 +6,6 @@
 #include "Modules/ModuleInterface.h"
 
 class AActor;
-class FEditorModeTools;
 
 DECLARE_DELEGATE_OneParam( FOnGetAllowedClasses, TArray<const UClass*>& );
 DECLARE_DELEGATE_OneParam( FOnActorSelected, AActor* );
@@ -42,8 +41,6 @@ private:
 	/** Handler for when the application is deactivated. */
 	void OnApplicationDeactivated(const bool IsActive) const;
 
-	static FEditorModeTools* GetLevelEditorModeManager();
-
 private:
 	FDelegateHandle OnApplicationDeactivatedHandle;
 };
diff --git a/Engine/Source/Editor/LevelInstanceEditor/Private/LevelInstanceEditorModule.cpp b/Engine/Source/Editor/LevelInstanceEditor/Private/LevelInstanceEditorModule.cpp
--- a/Engine/Source/Editor/LevelInstanceEditor/Private/LevelInstanceEditorModule.cpp
+++ b/Engine/Source/Editor/LevelInstanceEditor/Private/LevelInstanceEditorModule.cpp
@@ -918,15 +918,6 @@ void FLevelInstanceEditorModule::RegisterToFirstLevelEditor()
 	if (FirstLevelEditor.IsValid())
 	{
 		FirstLevelEditor->AddActorDetailsSCSEditorUICustomization(FLevelInstanceActorDetailsSCSEditorUICustomization::GetInstance());
-
-		FEditorModeTools& LevelEditorModeManager = FirstLevelEditor->GetEditorModeManager();
-		LevelEditorModeManager.OnEditorModeIDChanged().AddRaw(this, &FLevelInstanceEditorModule::OnEditorModeIDChanged);
-
-		// Create a Behavior source for the default EdModeTools (when we aren't in the LevelInstanceEditorMode)
-		DefaultBehaviorSource = ULevelInstanceEditorMode::CreateDefaultModeBehaviorSource(LevelEditorModeManager.GetInteractiveToolsContext());
-		LevelEditorModeManager.GetInteractiveToolsContext()->InputRouter->RegisterSource(DefaultBehaviorSource.GetInterface());
-		
-		RegisterLevelInstanceColumn();
 	}
 }
 
@@ -994,6 +985,17 @@ void FLevelInstanceEditorModule::StartupModule()
 		
 	FLevelInstanceEditorModeCommands::Register();
 
+	if (!IsRunningCommandlet())
+	{
+		GLevelEditorModeTools().OnEditorModeIDChanged().AddRaw(this, &FLevelInstanceEditorModule::OnEditorModeIDChanged);
+	
+		// Create a Behavior source for the default EdModeTools (when we aren't in the LevelInstanceEditorMode)
+		DefaultBehaviorSource = ULevelInstanceEditorMode::CreateDefaultModeBehaviorSource(GLevelEditorModeTools().GetInteractiveToolsContext());
+		GLevelEditorModeTools().GetInteractiveToolsContext()->InputRouter->RegisterSource(DefaultBehaviorSource.GetInterface());
+
+		RegisterLevelInstanceColumn();
+	}
+
 	ULevelInstanceSubsystem::RegisterPrimitiveColorHandler();
 }
 
@@ -1008,13 +1010,7 @@ void FLevelInstanceEditorModule::ShutdownModule()
 		if (TSharedPtr<ILevelEditor> FirstLevelEditor = LevelEditorModule.GetFirstLevelEditor())
 		{
 			FirstLevelEditor->RemoveActorDetailsSCSEditorUICustomization(FLevelInstanceActorDetailsSCSEditorUICustomization::GetInstance());
-			FirstLevelEditor->GetEditorModeManager().OnEditorModeIDChanged().RemoveAll(this);
-			FirstLevelEditor->GetEditorModeManager().GetInteractiveToolsContext()->InputRouter->DeregisterSource(DefaultBehaviorSource.GetInterface());
 		}
-
-		DefaultBehaviorSource = nullptr;
-		
-		UnregisterLevelInstanceColumn();
 	}
 
 	if (GEditor)
@@ -1023,6 +1019,17 @@ void FLevelInstanceEditorModule::ShutdownModule()
 	}
 
 	EditorLevelUtils::CanMoveActorToLevelDelegate.RemoveAll(this);
+
+	if (!IsRunningCommandlet())
+	{
+		if (GLevelEditorModeToolsIsValid())
+		{
+			GLevelEditorModeTools().OnEditorModeIDChanged().RemoveAll(this);
+			GLevelEditorModeTools().GetInteractiveToolsContext()->InputRouter->DeregisterSource(DefaultBehaviorSource.GetInterface());
+			DefaultBehaviorSource = nullptr;
+		}
+		UnregisterLevelInstanceColumn();
+	}
 }
 
 TSharedRef<ISceneOutlinerColumn> FLevelInstanceEditorModule::CreateLevelInstanceColumn(ISceneOutliner& SceneOutliner) const
@@ -1067,16 +1074,15 @@ void FLevelInstanceEditorModule::BroadcastTryExitEditorMode()
 
 void FLevelInstanceEditorModule::UpdateEditorMode(bool bActivated)
 {
-	FLevelEditorModule& LevelEditorModule = FModuleManager::GetModuleChecked<FLevelEditorModule>("LevelEditor");
-	if (TSharedPtr<ILevelEditor> FirstLevelEditor = LevelEditorModule.GetFirstLevelEditor())
+	if (!IsRunningCommandlet() && GLevelEditorModeToolsIsValid())
 	{
-		if (bActivated && !FirstLevelEditor->GetEditorModeManager().IsModeActive(ULevelInstanceEditorMode::EM_LevelInstanceEditorModeId))
+		if (bActivated && !GLevelEditorModeTools().IsModeActive(ULevelInstanceEditorMode::EM_LevelInstanceEditorModeId))
 		{
-			FirstLevelEditor->GetEditorModeManager().ActivateMode(ULevelInstanceEditorMode::EM_LevelInstanceEditorModeId);
+			GLevelEditorModeTools().ActivateMode(ULevelInstanceEditorMode::EM_LevelInstanceEditorModeId);
 		}
-		else if (!bActivated && FirstLevelEditor->GetEditorModeManager().IsModeActive(ULevelInstanceEditorMode::EM_LevelInstanceEditorModeId))
+		else if (!bActivated && GLevelEditorModeTools().IsModeActive(ULevelInstanceEditorMode::EM_LevelInstanceEditorModeId))
 		{
-			FirstLevelEditor->GetEditorModeManager().DeactivateMode(ULevelInstanceEditorMode::EM_LevelInstanceEditorModeId);
+			GLevelEditorModeTools().DeactivateMode(ULevelInstanceEditorMode::EM_LevelInstanceEditorModeId);
 		}
 	}
 }
diff --git a/Engine/Source/Editor/UnrealEd/Private/UnrealEdGlobals.cpp b/Engine/Source/Editor/UnrealEd/Private/UnrealEdGlobals.cpp
--- a/Engine/Source/Editor/UnrealEd/Private/UnrealEdGlobals.cpp
+++ b/Engine/Source/Editor/UnrealEd/Private/UnrealEdGlobals.cpp
@@ -47,22 +47,24 @@ DEFINE_LOG_CATEGORY_STATIC(LogUnrealEd, Log, All);
  */
 namespace Internal
 {
-	static TSharedPtr<FEditorModeTools> EditorModeToolsSingleton;
+	static TSharedPtr<FEditorModeTools>& GetGlobalModeManager()
+	{
+		checkf(!IsRunningCommandlet(), TEXT("The global mode manager should not be created or accessed in a commandlet environment. Check that your mode or module is not accessing the global mode tools or that scriptable features of modes have been moved to subsystems."));
+
+		static TSharedPtr<FEditorModeTools> EditorModeToolsSingleton = MakeShared<FEditorModeTools>();
+		return EditorModeToolsSingleton;
+	}
 };
 
 bool GLevelEditorModeToolsIsValid()
 {
-	return Internal::EditorModeToolsSingleton.IsValid();
+	return Internal::GetGlobalModeManager().IsValid();
 }
 
 FEditorModeTools& GLevelEditorModeTools()
 {
-	checkf(!IsRunningCommandlet(), TEXT("The global mode manager should not be created or accessed in a commandlet environment. Check that your mode or module is not accessing the global mode tools or that scriptable features of modes have been moved to subsystems."));
-	if (!ensureMsgf(Internal::EditorModeToolsSingleton.IsValid(), TEXT("The level editor is not started up yet. If you need to access the global mode manager early in the startup phase, please use FLevelEditorModule::OnLevelEditorCreated to gate the access.")))
-	{
-		Internal::EditorModeToolsSingleton = MakeShared<FEditorModeTools>();
-	}
-	return *Internal::EditorModeToolsSingleton.Get();
+	check(GLevelEditorModeToolsIsValid());
+	return *Internal::GetGlobalModeManager().Get();
 }
 
 FLevelEditorViewportClient* GCurrentLevelEditingViewportClient = NULL;
@@ -127,20 +129,16 @@ int32 EditorInit( IEngineLoop& EngineLoop )
 
 	// Set up the actor folders singleton
 	FActorFolders::Get();
-	if (!Internal::EditorModeToolsSingleton.IsValid())
-	{
-		Internal::EditorModeToolsSingleton = MakeShared<FEditorModeTools>();
-	}
 
 	// Initialize the misc editor
 	FUnrealEdMisc::Get().OnInit();
 	FCoreDelegates::OnExit.AddLambda([]()
 	{
 		// Shutdown the global static mode manager
-		if (Internal::EditorModeToolsSingleton.IsValid())
+		if (Internal::GetGlobalModeManager().IsValid())
 		{
 			GLevelEditorModeTools().SetDefaultMode(FBuiltinEditorModes::EM_Default);
-			Internal::EditorModeToolsSingleton.Reset();
+			Internal::GetGlobalModeManager().Reset();
 		}
 	});
 	
diff --git a/Engine/Source/Editor/UnrealEd/Public/Editor.h b/Engine/Source/Editor/UnrealEd/Public/Editor.h
--- a/Engine/Source/Editor/UnrealEd/Public/Editor.h
+++ b/Engine/Source/Editor/UnrealEd/Public/Editor.h
@@ -713,7 +713,6 @@ UNREALED_API class FEditorModeTools& GLevelEditorModeTools();
 /**
  * Checks if FEditorModeTools is valid
  */
-UE_DEPRECATED(5.5, "Checking the validity of the global mode manager is unnecessary. Instead use FLevelEditorModule::OnLevelEditorCreated to gate the access on the global mode manager.")
 UNREALED_API bool GLevelEditorModeToolsIsValid();
 
 namespace EditorUtilities
