LevelInstance: Speculative fix for SubSelection behavior crash for out of order Mode Enter/Exit.

diff --git a/Engine/Source/Editor/LevelInstanceEditor/Private/LevelInstanceEditorMode.cpp b/Engine/Source/Editor/LevelInstanceEditor/Private/LevelInstanceEditorMode.cpp
--- a/Engine/Source/Editor/LevelInstanceEditor/Private/LevelInstanceEditorMode.cpp
+++ b/Engine/Source/Editor/LevelInstanceEditor/Private/LevelInstanceEditorMode.cpp
@@ -215,6 +215,15 @@ void ULevelInstanceEditorMode::Enter()
 	
 	if (UEditorInteractiveToolsContext* InteractiveToolContext = GetInteractiveToolsContext(EToolsContextScope::EdMode))
 	{
+		// UEdMode::Exit() can be deferred to on Tick which can cause potentially out of order Enter/Exit calls.
+		// In the event that this does happen, we reregister the ModeBehaviorSource to prevent crashes, but ensure
+		// because the subsequent Exit will deregister the newly reregistered source and break viewport sub selection.
+		if (!ensureMsgf(ModeBehaviorSource == nullptr, TEXT("ModeBehaviorSource is already registered. Re-registering a new behavior source.")))
+		{
+			InteractiveToolContext->InputRouter->DeregisterSource(ModeBehaviorSource.GetInterface());
+			ModeBehaviorSource = nullptr;
+		}
+		
 		// Here we create a BehaviorSource specific to the Level Instance Editor Mode, for now it is the same type as the default one.
 		ModeBehaviorSource = CreateDefaultModeBehaviorSource(InteractiveToolContext);
 		InteractiveToolContext->InputRouter->RegisterSource(ModeBehaviorSource.GetInterface());
@@ -228,6 +237,7 @@ void ULevelInstanceEditorMode::Exit()
 	if (UEditorInteractiveToolsContext* InteractiveToolContext = GetInteractiveToolsContext(EToolsContextScope::EdMode))
 	{
 		InteractiveToolContext->InputRouter->DeregisterSource(ModeBehaviorSource.GetInterface());
+		ModeBehaviorSource = nullptr;
 	}
 
 	UEdMode::Exit();
