Ensure the default HLOD Layer setup is duplicated in all code paths on VK projects creation
* Code Path used during VK tests was leaving maps with the default HLOD Layer asset assigned, leading to validation issue
* Also fixed an issue where the default HLOD Layer was improperly outered when converting an existing map to World Partition

diff --git a/Engine/Source/Editor/LevelInstanceEditor/Private/LevelInstanceEditorModule.cpp b/Engine/Source/Editor/LevelInstanceEditor/Private/LevelInstanceEditorModule.cpp
--- a/Engine/Source/Editor/LevelInstanceEditor/Private/LevelInstanceEditorModule.cpp
+++ b/Engine/Source/Editor/LevelInstanceEditor/Private/LevelInstanceEditorModule.cpp
@@ -45,6 +45,7 @@
 #include "SNewLevelInstanceDialog.h"
 #include "MessageLogModule.h"
 #include "Settings/EditorExperimentalSettings.h"
+#include "WorldPartition/HLOD/HLODLayer.h"
 #include "WorldPartition/WorldPartitionConverter.h"
 #include "WorldPartition/WorldPartitionActorLoaderInterface.h"
 #include "ScopedTransaction.h"
@@ -777,6 +778,16 @@ struct FLevelInstanceMenuUtils
 			TSet<UPackage*> PackagesToSaveSet(PackagesToSave);
 
 			PackagesToSaveSet.Add(WorldAsset->GetPackage());
+			
+			UHLODLayer* NewHLODLayer = WorldAsset->GetWorldPartition()->GetDefaultHLODLayer();
+			while (NewHLODLayer && NewHLODLayer->GetPackage())
+			{
+				if (NewHLODLayer->GetPackage()->IsDirty())
+				{
+					PackagesToSaveSet.Emplace(NewHLODLayer->GetPackage());
+				}
+				NewHLODLayer = NewHLODLayer->GetParentLayer();
+			}
 
 			const bool bPromptUserToSave = false;
 			const bool bSaveMapPackages = true;
diff --git a/Engine/Source/Runtime/Engine/Private/WorldPartition/HLOD/HLODLayer.cpp b/Engine/Source/Runtime/Engine/Private/WorldPartition/HLOD/HLODLayer.cpp
--- a/Engine/Source/Runtime/Engine/Private/WorldPartition/HLOD/HLODLayer.cpp
+++ b/Engine/Source/Runtime/Engine/Private/WorldPartition/HLOD/HLODLayer.cpp
@@ -77,6 +77,7 @@ UHLODLayer* UHLODLayer::DuplicateHLODLayersSetup(UHLODLayer* HLODLayer, const FS
 		UPackage* Package = CreatePackage(*PackageName);
 		// In case Package already exists setting this flag will allow overwriting it
 		Package->MarkAsFullyLoaded();
+		Package->SetDirtyFlag(true);
 
 		FObjectDuplicationParameters ObjParameters(CurrentHLODLayer, Package);
 		ObjParameters.DestName = FName(Prefix + TEXT("_") + CurrentHLODLayer->GetName());
diff --git a/Engine/Source/Runtime/Engine/Private/WorldPartition/WorldPartition.cpp b/Engine/Source/Runtime/Engine/Private/WorldPartition/WorldPartition.cpp
--- a/Engine/Source/Runtime/Engine/Private/WorldPartition/WorldPartition.cpp
+++ b/Engine/Source/Runtime/Engine/Private/WorldPartition/WorldPartition.cpp
@@ -1101,7 +1101,11 @@ UWorldPartition* UWorldPartition::CreateOrRepairWorldPartition(AWorldSettings* W
 
 		WorldSettings->MarkPackageDirty();
 
-		WorldPartition->DefaultHLODLayer = UHLODLayer::GetEngineDefaultHLODLayersSetup();
+		// Assign a copy of the default HLOD layers setup to this new map
+		if (UHLODLayer* HLODLayerTemplate = UHLODLayer::GetEngineDefaultHLODLayersSetup())
+		{
+			WorldPartition->DefaultHLODLayer = UHLODLayer::DuplicateHLODLayersSetup(HLODLayerTemplate, WorldPartition->GetPackage()->GetName(), WorldPartition->GetWorld()->GetName());
+		}
 
 		AWorldDataLayers* WorldDataLayers = OuterWorld->GetWorldDataLayers();
 		if (!WorldDataLayers)
