Fixed type incorrectness constructing when constructing TSoftObjectPtr<UBlueprint> with a UObject*.

diff --git a/Engine/Source/Editor/LevelInstanceEditor/Private/LevelInstanceEditorModule.cpp b/Engine/Source/Editor/LevelInstanceEditor/Private/LevelInstanceEditorModule.cpp
--- a/Engine/Source/Editor/LevelInstanceEditor/Private/LevelInstanceEditorModule.cpp
+++ b/Engine/Source/Editor/LevelInstanceEditor/Private/LevelInstanceEditorModule.cpp
@@ -563,7 +563,7 @@ struct FLevelInstanceMenuUtils
 				;
 				if (APackedLevelActor* PackedLevelActor = Cast<APackedLevelActor>(ContextLevelInstance))
 				{
-					if (TSoftObjectPtr<UBlueprint> BlueprintAsset = PackedLevelActor->GetClass()->ClassGeneratedBy.Get())
+					if (TSoftObjectPtr<UBlueprint> BlueprintAsset = Cast<UBlueprint>(PackedLevelActor->GetClass()->ClassGeneratedBy.Get()))
 					{
 						FToolUIAction UIAction;
 						UIAction.ExecuteAction.BindLambda([ContextLevelInstance, BlueprintAsset](const FToolMenuContext& MenuContext)
