Fixed implicit cast to TSoftObjectPtr<UWorld>.

diff --git a/Engine/Source/Editor/LevelInstanceEditor/Private/LevelInstanceActorFactory.cpp b/Engine/Source/Editor/LevelInstanceEditor/Private/LevelInstanceActorFactory.cpp
--- a/Engine/Source/Editor/LevelInstanceEditor/Private/LevelInstanceActorFactory.cpp
+++ b/Engine/Source/Editor/LevelInstanceEditor/Private/LevelInstanceActorFactory.cpp
@@ -14,7 +14,7 @@ ULevelInstanceActorFactory::ULevelInstanceActorFactory(const FObjectInitializer&
 void ULevelInstanceActorFactory::PostSpawnActor(UObject* Asset, AActor* NewActor)
 {
 	ILevelInstanceInterface* LevelInstanceInterface = CastChecked<ILevelInstanceInterface>(NewActor);
-	LevelInstanceInterface->SetWorldAsset(Asset);
+	LevelInstanceInterface->SetWorldAsset(CastChecked<UWorld>(Asset));
 	LevelInstanceInterface->LoadLevelInstance();
 }
 
