LevelInstance: Prevent mouse wheel capture if bIsViewportSubSelectionEnabled is false.

diff --git a/Engine/Source/Editor/LevelInstanceEditor/Private/LevelInstanceEditorMode.cpp b/Engine/Source/Editor/LevelInstanceEditor/Private/LevelInstanceEditorMode.cpp
--- a/Engine/Source/Editor/LevelInstanceEditor/Private/LevelInstanceEditorMode.cpp
+++ b/Engine/Source/Editor/LevelInstanceEditor/Private/LevelInstanceEditorMode.cpp
@@ -175,6 +175,10 @@ void ULevelInstanceEditorBehaviorSource::Initialize(UEditorInteractiveToolsConte
 	UMouseWheelInputBehavior* MouseWheelInputBehavior = NewObject<UMouseWheelInputBehavior>();
 	MouseWheelBehaviorTarget = MakeUnique<FMouseWheelBehaviorTarget>(InteractiveToolsContext);
 	MouseWheelInputBehavior->Initialize(MouseWheelBehaviorTarget.Get());
+	MouseWheelInputBehavior->ModifierCheckFunc = [](const FInputDeviceState&)
+	{
+		return GetDefault<ULevelInstanceEditorPerProjectUserSettings>()->bIsViewportSubSelectionEnabled;
+	};
 	InputBehaviorSet->Add(MouseWheelInputBehavior);
 }
 
