[Level Instance] Fix DefaultBehaviorSource leak caused by OnLevelEditorCreated() being called when changing the editor layout

#rnx

diff --git a/Engine/Source/Editor/LevelInstanceEditor/Private/LevelInstanceEditorModule.cpp b/Engine/Source/Editor/LevelInstanceEditor/Private/LevelInstanceEditorModule.cpp
--- a/Engine/Source/Editor/LevelInstanceEditor/Private/LevelInstanceEditorModule.cpp
+++ b/Engine/Source/Editor/LevelInstanceEditor/Private/LevelInstanceEditorModule.cpp
@@ -927,6 +927,13 @@ void FLevelInstanceEditorModule::RegisterToFirstLevelEditor()
 		LevelEditorModeManager.GetInteractiveToolsContext()->InputRouter->RegisterSource(DefaultBehaviorSource.GetInterface());
 		
 		RegisterLevelInstanceColumn();
+
+		// Make sure to unregister because changing the layout will callback on this again.
+		// 
+		// This works because we aren't actually hooking ourselves to the ILevelEditor but on managers that are shared by the different instances
+		// of ILevelEditor. Ideally we could listen to an event when a ILevelEditor gets destroyed to unregister ourselves and continue to listen
+		// to this event to re-register ourselves
+		LevelEditorModule.OnLevelEditorCreated().RemoveAll(this);
 	}
 }
 
