Fixed crash when creating a new Level Instance Actor from the create actor panel

diff --git a/Engine/Source/Editor/LevelInstanceEditor/Private/LevelInstanceActorFactory.cpp b/Engine/Source/Editor/LevelInstanceEditor/Private/LevelInstanceActorFactory.cpp
--- a/Engine/Source/Editor/LevelInstanceEditor/Private/LevelInstanceActorFactory.cpp
+++ b/Engine/Source/Editor/LevelInstanceEditor/Private/LevelInstanceActorFactory.cpp
@@ -14,8 +14,11 @@ ULevelInstanceActorFactory::ULevelInstanceActorFactory(const FObjectInitializer&
 void ULevelInstanceActorFactory::PostSpawnActor(UObject* Asset, AActor* NewActor)
 {
 	ILevelInstanceInterface* LevelInstanceInterface = CastChecked<ILevelInstanceInterface>(NewActor);
-	LevelInstanceInterface->SetWorldAsset(CastChecked<UWorld>(Asset));
-	LevelInstanceInterface->LoadLevelInstance();
+	if (UWorld* WorldAsset = Cast<UWorld>(Asset))
+	{
+		LevelInstanceInterface->SetWorldAsset(WorldAsset);
+		LevelInstanceInterface->LoadLevelInstance();
+	}
 }
 
 bool ULevelInstanceActorFactory::CanCreateActorFrom(const FAssetData& AssetData, FText& OutErrorMsg)
