Avoid interacting with EditorModeManager singleton during editor shutdown, as it's destroyed in AppPreExit, which is before ShutdownModule is called during editor shutdown (AppExit)

diff --git a/Engine/Source/Editor/LevelInstanceEditor/Private/LevelInstanceEditorModule.cpp b/Engine/Source/Editor/LevelInstanceEditor/Private/LevelInstanceEditorModule.cpp
--- a/Engine/Source/Editor/LevelInstanceEditor/Private/LevelInstanceEditorModule.cpp
+++ b/Engine/Source/Editor/LevelInstanceEditor/Private/LevelInstanceEditorModule.cpp
@@ -1015,8 +1015,11 @@ void FLevelInstanceEditorModule::ShutdownModule()
 		if (TSharedPtr<ILevelEditor> FirstLevelEditor = LevelEditorModule.GetFirstLevelEditor())
 		{
 			FirstLevelEditor->RemoveActorDetailsSCSEditorUICustomization(FLevelInstanceActorDetailsSCSEditorUICustomization::GetInstance());
-			FirstLevelEditor->GetEditorModeManager().OnEditorModeIDChanged().RemoveAll(this);
-			FirstLevelEditor->GetEditorModeManager().GetInteractiveToolsContext()->InputRouter->DeregisterSource(DefaultBehaviorSource.GetInterface());
+			if (!IsEngineExitRequested())
+			{
+				FirstLevelEditor->GetEditorModeManager().OnEditorModeIDChanged().RemoveAll(this);
+				FirstLevelEditor->GetEditorModeManager().GetInteractiveToolsContext()->InputRouter->DeregisterSource(DefaultBehaviorSource.GetInterface());
+			}
 		}
 
 		DefaultBehaviorSource = nullptr;
