Level Instance Module: Add a "Browse to Asset" override for level instance to override the default behavior of browsing to both the level instance asset and the current persistent level asset

diff --git a/Engine/Source/Editor/LevelInstanceEditor/Private/LevelInstanceEditorModule.cpp b/Engine/Source/Editor/LevelInstanceEditor/Private/LevelInstanceEditorModule.cpp
--- a/Engine/Source/Editor/LevelInstanceEditor/Private/LevelInstanceEditorModule.cpp
+++ b/Engine/Source/Editor/LevelInstanceEditor/Private/LevelInstanceEditorModule.cpp
@@ -52,6 +52,7 @@
 #include "Tools/EdModeInteractiveToolsContext.h"
 #include "SceneOutlinerModule.h"
 #include "SceneOutlinerFwd.h"
+#include "Subsystems/BrowseToAssetOverrideSubsystem.h"
 
 IMPLEMENT_MODULE( FLevelInstanceEditorModule, LevelInstanceEditor );
 
@@ -937,6 +938,13 @@ void FLevelInstanceEditorModule::RegisterToFirstLevelEditor()
 	}
 }
 
+FName FLevelInstanceEditorModule::GetBrowseToLevelInstanceAsset(const UObject* Object) const
+{
+	// Level instances browse to both the level instance asset and the current level asset by default, while we only want to browse to the former
+	const ILevelInstanceInterface* LevelInstanceInterface = CastChecked<ILevelInstanceInterface>(Object);
+	return FName(LevelInstanceInterface->GetWorldAssetPackage());
+}
+
 void FLevelInstanceEditorModule::StartupModule()
 {
 	FLevelEditorModule& LevelEditorModule = FModuleManager::LoadModuleChecked<FLevelEditorModule>("LevelEditor");
@@ -1002,10 +1010,21 @@ void FLevelInstanceEditorModule::StartupModule()
 	FLevelInstanceEditorModeCommands::Register();
 
 	ULevelInstanceSubsystem::RegisterPrimitiveColorHandler();
+
+	if (UBrowseToAssetOverrideSubsystem* BrowseToAssetOverrideSubsystem = UBrowseToAssetOverrideSubsystem::Get())
+	{
+		BrowseToAssetOverrideSubsystem->RegisterBrowseToAssetOverrideForInterface<ILevelInstanceInterface>(
+			FBrowseToAssetOverrideDelegate::CreateRaw(this, &FLevelInstanceEditorModule::GetBrowseToLevelInstanceAsset));
+	}
 }
 
 void FLevelInstanceEditorModule::ShutdownModule()
 {
+	if (UBrowseToAssetOverrideSubsystem* BrowseToAssetOverrideSubsystem = UBrowseToAssetOverrideSubsystem::Get())
+	{
+		BrowseToAssetOverrideSubsystem->UnregisterBrowseToAssetOverrideForInterface<ILevelInstanceInterface>();
+	}
+	
 	ULevelInstanceSubsystem::UnregisterPrimitiveColorHandler();
 
 	if (FModuleManager::Get().IsModuleLoaded("LevelEditor"))
diff --git a/Engine/Source/Editor/LevelInstanceEditor/Private/LevelInstanceEditorModule.h b/Engine/Source/Editor/LevelInstanceEditor/Private/LevelInstanceEditorModule.h
--- a/Engine/Source/Editor/LevelInstanceEditor/Private/LevelInstanceEditorModule.h
+++ b/Engine/Source/Editor/LevelInstanceEditor/Private/LevelInstanceEditorModule.h
@@ -58,6 +58,7 @@ private:
 		
 	void OnLevelEditorCreated(TSharedPtr<ILevelEditor> InLevelEditor);
 	void RegisterToFirstLevelEditor();
+	FName GetBrowseToLevelInstanceAsset(const UObject* Object) const;
 
 	void RegisterLevelInstanceColumn();
 	void UnregisterLevelInstanceColumn();
diff --git a/Engine/Source/Editor/UnrealEd/Private/Subsystems/BrowseToAssetOverrideSubsystem.cpp b/Engine/Source/Editor/UnrealEd/Private/Subsystems/BrowseToAssetOverrideSubsystem.cpp
--- a/Engine/Source/Editor/UnrealEd/Private/Subsystems/BrowseToAssetOverrideSubsystem.cpp
+++ b/Engine/Source/Editor/UnrealEd/Private/Subsystems/BrowseToAssetOverrideSubsystem.cpp
@@ -43,6 +43,24 @@ FName UBrowseToAssetOverrideSubsystem::GetBrowseToAssetOverride(const UObject* O
 		}
 	}
 
+	// Query all the class interfaces to see if there's a valid per-interface override for this instance
+	if (PerInterfaceOverrides.Num() > 0)
+	{
+		UClass* ObjectClass = Object->GetClass();
+		for (const FImplementedInterface& Interface : ObjectClass->Interfaces)
+		{
+			if (const FBrowseToAssetOverrideDelegate* CallbackPtr = PerInterfaceOverrides.Find(Interface.Class->GetClassPathName()))
+			{
+				FName CallbackBrowseToAssetOverride = CallbackPtr->IsBound() ? CallbackPtr->Execute(Object) : FName();
+				if (!CallbackBrowseToAssetOverride.IsNone())
+				{
+					return CallbackBrowseToAssetOverride;
+				}
+			}
+		}		
+	}
+
+
 	return FName();
 }
 
@@ -55,3 +73,13 @@ void UBrowseToAssetOverrideSubsystem::UnregisterBrowseToAssetOverrideForClass(co
 {
 	PerClassOverrides.Remove(Class);
 }
+
+void UBrowseToAssetOverrideSubsystem::RegisterBrowseToAssetOverrideForInterface(const FTopLevelAssetPath& Interface, FBrowseToAssetOverrideDelegate&& Callback)
+{
+	PerInterfaceOverrides.Add(Interface, MoveTemp(Callback));
+}
+
+void UBrowseToAssetOverrideSubsystem::UnregisterBrowseToAssetOverrideForInterface(const FTopLevelAssetPath& Interface)
+{
+	PerInterfaceOverrides.Remove(Interface);
+}
diff --git a/Engine/Source/Editor/UnrealEd/Public/Subsystems/BrowseToAssetOverrideSubsystem.h b/Engine/Source/Editor/UnrealEd/Public/Subsystems/BrowseToAssetOverrideSubsystem.h
--- a/Engine/Source/Editor/UnrealEd/Public/Subsystems/BrowseToAssetOverrideSubsystem.h
+++ b/Engine/Source/Editor/UnrealEd/Public/Subsystems/BrowseToAssetOverrideSubsystem.h
@@ -46,7 +46,31 @@ public:
 	{
 		UnregisterBrowseToAssetOverrideForClass(ClassType::StaticClass()->GetClassPathName());
 	}
+
+	/**
+	 * Register a per-interface override for the "browse to asset" resolution.
+	 * @note Callback should return a package name, or None if there is no override.
+	 */
+	UNREALED_API void RegisterBrowseToAssetOverrideForInterface(const FTopLevelAssetPath& Interface, FBrowseToAssetOverrideDelegate&& Callback);
+
+	template <typename InterfaceType>
+	void RegisterBrowseToAssetOverrideForInterface(FBrowseToAssetOverrideDelegate&& Callback)
+	{
+		RegisterBrowseToAssetOverrideForInterface(InterfaceType::UClassType::StaticClass()->GetClassPathName(), MoveTemp(Callback));
+	}
+
+	/**
+	 * Unregister a per-interface override for the "browse to asset" resolution.
+	 */
+	UNREALED_API void UnregisterBrowseToAssetOverrideForInterface(const FTopLevelAssetPath& Interface);
+
+	template <typename InterfaceType>
+	void UnregisterBrowseToAssetOverrideForInterface()
+	{
+		UnregisterBrowseToAssetOverrideForInterface(InterfaceType::UClassType::StaticClass()->GetClassPathName());
+	}
 	
 private:
 	TMap<FTopLevelAssetPath, FBrowseToAssetOverrideDelegate> PerClassOverrides;
+	TMap<FTopLevelAssetPath, FBrowseToAssetOverrideDelegate> PerInterfaceOverrides;
 };
