[Level Instance] Always show the Override item for Level Instances, but explain why it is disabled when it is

#rnx
#Synced-CL 49104468

diff --git a/Engine/Source/Editor/LevelInstanceEditor/Private/LevelInstanceEditorModule.cpp b/Engine/Source/Editor/LevelInstanceEditor/Private/LevelInstanceEditorModule.cpp
--- a/Engine/Source/Editor/LevelInstanceEditor/Private/LevelInstanceEditorModule.cpp
+++ b/Engine/Source/Editor/LevelInstanceEditor/Private/LevelInstanceEditorModule.cpp
@@ -136,6 +136,7 @@ struct FLevelInstanceMenuUtils
 		FText EntryLabel = bSingleEntry ? LOCTEXT("OverrideLevelInstances", "Override") : FText::FromString(LevelInstance->GetWorldAsset().GetAssetName());
 		if (bCanEdit)
 		{
+			// EntryDesc contains the reason the menu cannot be enabled, when CanEnterEditPropertyOverrides would return false. It returned true, so the item will be enabled, and we set another tooltip
 			FText EntryActionDesc = LOCTEXT("EditLevelInstancesPropertyOverridesTooltip", "Edit only this level instance, without changing the level asset or any other level instances.");
 			EntryDesc = FText::Format(LOCTEXT("OverrideLevelInstanceName", "{0}\n\nActor name: {1}\nAsset path: {2}"), EntryActionDesc, FText::FromString(LevelInstanceActor->GetActorLabel()), FText::FromString(LevelInstance->GetWorldAssetPackage()));
 		}
@@ -188,11 +189,6 @@ struct FLevelInstanceMenuUtils
 
 	static void CreateEditPropertyOverridesMenu(UToolMenu* Menu, AActor* ContextActor)
 	{
-		if (!ULevelInstanceSettings::Get()->IsPropertyOverrideEnabled())
-		{
-			return;
-		}
-
 		if (ULevelInstanceSubsystem* LevelInstanceSubsystem = ContextActor->GetWorld()->GetSubsystem<ULevelInstanceSubsystem>())
 		{
 			TArray<ILevelInstanceInterface*> LevelInstanceHierarchy;
diff --git a/Engine/Source/Runtime/Engine/Private/LevelInstance/LevelInstanceSettings.cpp b/Engine/Source/Runtime/Engine/Private/LevelInstance/LevelInstanceSettings.cpp
--- a/Engine/Source/Runtime/Engine/Private/LevelInstance/LevelInstanceSettings.cpp
+++ b/Engine/Source/Runtime/Engine/Private/LevelInstance/LevelInstanceSettings.cpp
@@ -14,9 +14,19 @@ ULevelInstanceSettings::ULevelInstanceSettings()
 }
 
 #if WITH_EDITOR
+bool ULevelInstanceSettings::IsLevelInstancePropertyOverrideSettingEnabled() const
+{
+	return GetDefault<UEditorExperimentalSettings>()->bEnableLevelInstancePropertyOverrides;
+}
+
+bool ULevelInstanceSettings::IsPropertyOverridePolicyDefined() const
+{
+	return !!PropertyOverridePolicy;
+}
+
 bool ULevelInstanceSettings::IsPropertyOverrideEnabled() const
 {
-	return GetDefault<UEditorExperimentalSettings>()->bEnableLevelInstancePropertyOverrides && !!PropertyOverridePolicy;
+	return IsLevelInstancePropertyOverrideSettingEnabled() && IsPropertyOverridePolicyDefined();
 }
 
 bool ULevelInstanceSettings::IsLevelInstanceEditCompatibleWithLandscapeEdit() const
diff --git a/Engine/Source/Runtime/Engine/Private/LevelInstance/LevelInstanceSubsystem.cpp b/Engine/Source/Runtime/Engine/Private/LevelInstance/LevelInstanceSubsystem.cpp
--- a/Engine/Source/Runtime/Engine/Private/LevelInstance/LevelInstanceSubsystem.cpp
+++ b/Engine/Source/Runtime/Engine/Private/LevelInstance/LevelInstanceSubsystem.cpp
@@ -2355,11 +2355,20 @@ bool ULevelInstanceSubsystem::CanEditLevelInstancePropertyOverrides(const ILevel
 		return false;
 	}
 
-	if (!ULevelInstanceSettings::Get()->IsPropertyOverrideEnabled())
+	if (!ULevelInstanceSettings::Get()->IsLevelInstancePropertyOverrideSettingEnabled())
 	{
 		if (OutReason)
 		{
-			*OutReason = LOCTEXT("LevelInstanceNotEnabledPropertyOverrides", "Level Instance property override feature is not enabled");
+			*OutReason = LOCTEXT("LevelInstanceNotEnabledPropertyOverrides", "Check the \"Enable Level Instance Property Overrides\" Editor Setting");
+		}
+		return false;
+	}
+
+	if (!ULevelInstanceSettings::Get()->IsPropertyOverridePolicyDefined())
+	{
+		if (OutReason)
+		{
+			*OutReason = LOCTEXT("LevelInstanceNotEnabledPropertyOverrides", "Provide an override policy class in your DefaultEngine.ini file.\nFor ex. under [/Script/Engine.LevelInstanceSettings], add PropertyOverridePolicyClass=/Script/Engine.LevelInstancePropertyOverrideSamplePolicy");
 		}
 		return false;
 	}
diff --git a/Engine/Source/Runtime/Engine/Public/LevelInstance/LevelInstanceSettings.h b/Engine/Source/Runtime/Engine/Public/LevelInstance/LevelInstanceSettings.h
--- a/Engine/Source/Runtime/Engine/Public/LevelInstance/LevelInstanceSettings.h
+++ b/Engine/Source/Runtime/Engine/Public/LevelInstance/LevelInstanceSettings.h
@@ -20,7 +20,10 @@ public:
 	static ULevelInstanceSettings* Get() { return CastChecked<ULevelInstanceSettings>(ULevelInstanceSettings::StaticClass()->GetDefaultObject()); }
 
 #if WITH_EDITOR
+	UE_API bool IsLevelInstancePropertyOverrideSettingEnabled() const;
+	UE_API bool IsPropertyOverridePolicyDefined() const;
 	UE_API bool IsPropertyOverrideEnabled() const;
+
 	UE_API bool IsLevelInstanceEditCompatibleWithLandscapeEdit() const;
 
 	UE_DEPRECATED(5.5, "Use SetIsLevelInstanceDisabled instead.")
