Fixed 'Level Instance Overrides' column not showing and not available when loading the editor after enabling level instance property overrides.
Also fixed the fact that this column could be available in any outliner.

#Synced-CL 49593614

diff --git a/Engine/Source/Editor/LevelInstanceEditor/Private/LevelInstanceEditorModule.cpp b/Engine/Source/Editor/LevelInstanceEditor/Private/LevelInstanceEditorModule.cpp
--- a/Engine/Source/Editor/LevelInstanceEditor/Private/LevelInstanceEditorModule.cpp
+++ b/Engine/Source/Editor/LevelInstanceEditor/Private/LevelInstanceEditorModule.cpp
@@ -53,6 +53,7 @@
 #include "SceneOutlinerModule.h"
 #include "SceneOutlinerFwd.h"
 #include "Subsystems/BrowseToAssetOverrideSubsystem.h"
+#include "SceneOutlinerPublicTypes.h"
 
 IMPLEMENT_MODULE( FLevelInstanceEditorModule, LevelInstanceEditor );
 
@@ -898,6 +899,18 @@ void FLevelInstanceEditorModule::OnLevelEditorCreated(TSharedPtr<ILevelEditor> I
 	RegisterToFirstLevelEditor();
 }
 
+void FLevelInstanceEditorModule::OnCreateActorBrowserColumns(FSceneOutlinerInitializationOptions& InitOptions, UWorld* World)
+{
+	if (GetDefault<ULevelInstanceSettings>()->IsPropertyOverrideEnabled())
+	{
+		FSceneOutlinerColumnInfo ColumnInfo(ESceneOutlinerColumnVisibility::Invisible, 8,
+			FCreateSceneOutlinerColumn::CreateRaw(this, &FLevelInstanceEditorModule::CreateLevelInstanceColumn),
+			true, TOptional<float>(), LOCTEXT("LevelInstanceColumnName", "Level Instance Overrides"));
+
+		InitOptions.ColumnMap.Add(FLevelInstanceSceneOutlinerColumn::GetID(), ColumnInfo);
+	}
+}
+
 void FLevelInstanceEditorModule::RegisterToFirstLevelEditor()
 {
 	FLevelEditorModule& LevelEditorModule = FModuleManager::GetModuleChecked<FLevelEditorModule>("LevelEditor");
@@ -913,8 +926,6 @@ void FLevelInstanceEditorModule::RegisterToFirstLevelEditor()
 		DefaultBehaviorSource = ULevelInstanceEditorMode::CreateDefaultModeBehaviorSource(LevelEditorModeManager.GetInteractiveToolsContext());
 		LevelEditorModeManager.GetInteractiveToolsContext()->InputRouter->RegisterSource(DefaultBehaviorSource.GetInterface());
 		
-		RegisterLevelInstanceColumn();
-
 		// Make sure to unregister because changing the layout will callback on this again.
 		// 
 		// This works because we aren't actually hooking ourselves to the ILevelEditor but on managers that are shared by the different instances
@@ -942,6 +953,9 @@ void FLevelInstanceEditorModule::StartupModule()
 	{
 		LevelEditorModule.OnLevelEditorCreated().AddRaw(this, &FLevelInstanceEditorModule::OnLevelEditorCreated);
 	}
+
+	FSceneOutlinerModule& SceneOutlinerModule = FModuleManager::LoadModuleChecked<FSceneOutlinerModule>("SceneOutliner");
+	SceneOutlinerModule.OnCreateActorBrowserColumns().AddRaw(this, &FLevelInstanceEditorModule::OnCreateActorBrowserColumns);
 			
 	ExtendContextMenu();
 
@@ -1028,8 +1042,11 @@ void FLevelInstanceEditorModule::ShutdownModule()
 		}
 
 		DefaultBehaviorSource = nullptr;
-		
-		UnregisterLevelInstanceColumn();
+	}
+
+	if (FSceneOutlinerModule* SceneOutlinerModule = FModuleManager::GetModulePtr<FSceneOutlinerModule>("SceneOutliner"))
+	{
+		SceneOutlinerModule->OnCreateActorBrowserColumns().RemoveAll(this);
 	}
 
 	if (GEditor)
@@ -1045,28 +1062,6 @@ TSharedRef<ISceneOutlinerColumn> FLevelInstanceEditorModule::CreateLevelInstance
 	return MakeShareable(new FLevelInstanceSceneOutlinerColumn(SceneOutliner));
 }
 
-void FLevelInstanceEditorModule::RegisterLevelInstanceColumn()
-{
-	if (GetDefault<ULevelInstanceSettings>()->IsPropertyOverrideEnabled())
-	{
-		FSceneOutlinerModule& SceneOutlinerModule = FModuleManager::LoadModuleChecked<FSceneOutlinerModule>("SceneOutliner");
-
-		FSceneOutlinerColumnInfo ColumnInfo(ESceneOutlinerColumnVisibility::Invisible, 8,
-			FCreateSceneOutlinerColumn::CreateRaw(this, &FLevelInstanceEditorModule::CreateLevelInstanceColumn),
-			true, TOptional<float>(), LOCTEXT("LevelInstanceColumnName", "Level Instance Overrides"));
-
-		SceneOutlinerModule.RegisterDefaultColumnType<FLevelInstanceSceneOutlinerColumn>(ColumnInfo);
-	}
-}
-
-void FLevelInstanceEditorModule::UnregisterLevelInstanceColumn()
-{
-	if (FSceneOutlinerModule* SceneOutlinerModulePtr = FModuleManager::GetModulePtr<FSceneOutlinerModule>("SceneOutliner"))
-	{
-		SceneOutlinerModulePtr->UnRegisterColumnType<FLevelInstanceSceneOutlinerColumn>();
-	}
-}
-
 void FLevelInstanceEditorModule::OnEditorModeIDChanged(const FEditorModeID& InModeID, bool bIsEnteringMode)
 {
 	if (InModeID == ULevelInstanceEditorMode::EM_LevelInstanceEditorModeId && !bIsEnteringMode)
diff --git a/Engine/Source/Editor/LevelInstanceEditor/Private/LevelInstanceEditorModule.h b/Engine/Source/Editor/LevelInstanceEditor/Private/LevelInstanceEditorModule.h
--- a/Engine/Source/Editor/LevelInstanceEditor/Private/LevelInstanceEditorModule.h
+++ b/Engine/Source/Editor/LevelInstanceEditor/Private/LevelInstanceEditorModule.h
@@ -15,6 +15,7 @@ enum class EMapChangeType : uint8;
 class ILevelEditor;
 class ISceneOutliner;
 class ISceneOutlinerColumn;
+struct FSceneOutlinerInitializationOptions;
 
 /**
  * The module holding all of the UI related pieces for LevelInstance management
@@ -63,11 +64,10 @@ private:
 	void ExtendContextMenu();
 		
 	void OnLevelEditorCreated(TSharedPtr<ILevelEditor> InLevelEditor);
+	void OnCreateActorBrowserColumns(FSceneOutlinerInitializationOptions& InitOptions, UWorld* World);
 	void RegisterToFirstLevelEditor();
 	FTopLevelAssetPath GetBrowseToLevelInstanceAsset(const UObject* Object) const;
 
-	void RegisterLevelInstanceColumn();
-	void UnregisterLevelInstanceColumn();
 	TSharedRef<ISceneOutlinerColumn> CreateLevelInstanceColumn(ISceneOutliner& SceneOutliner) const;
 
 	class FPropertyEditorPolicy : public PropertyEditorPolicy::IEditConstPolicy, 
diff --git a/Engine/Source/Editor/SceneOutliner/Private/SceneOutlinerModule.cpp b/Engine/Source/Editor/SceneOutliner/Private/SceneOutlinerModule.cpp
--- a/Engine/Source/Editor/SceneOutliner/Private/SceneOutlinerModule.cpp
+++ b/Engine/Source/Editor/SceneOutliner/Private/SceneOutlinerModule.cpp
@@ -232,6 +232,8 @@ void FSceneOutlinerModule::CreateActorBrowserColumns(FSceneOutlinerInitializatio
 
 	CreateWorldPartitionColumns(InInitOptions, InWorld);
 	CreateActorInfoColumns(InInitOptions, InWorld);
+
+	CreateActorBrowserColumnsDelegate.Broadcast(InInitOptions, InWorld);
 }
 
 void FSceneOutlinerModule::CreateWorldPartitionColumns(FSceneOutlinerInitializationOptions& InInitOptions, UWorld* WorldPtr) const
diff --git a/Engine/Source/Editor/SceneOutliner/Public/SceneOutlinerModule.h b/Engine/Source/Editor/SceneOutliner/Public/SceneOutlinerModule.h
--- a/Engine/Source/Editor/SceneOutliner/Public/SceneOutlinerModule.h
+++ b/Engine/Source/Editor/SceneOutliner/Public/SceneOutlinerModule.h
@@ -84,6 +84,10 @@ public:
 	DECLARE_MULTICAST_DELEGATE(FOnColumnPermissionListChanged);
 	FOnColumnPermissionListChanged& OnColumnPermissionListChanged() { return ColumnPermissionListChanged; }
 
+	/** Delegate that broadcasts when creating actor browser columns. */
+	DECLARE_MULTICAST_DELEGATE_TwoParams(FOnCreateActorBrowserColumns, FSceneOutlinerInitializationOptions&, UWorld*);
+	FOnCreateActorBrowserColumns& OnCreateActorBrowserColumns() { return CreateActorBrowserColumnsDelegate; }
+
 public:
 	/** Register a new type of column available to all scene outliners */
 	template< typename T >
@@ -151,6 +155,9 @@ private:
 
 	/** Delegate that broadcasts when column permission list changes. */
 	FOnColumnPermissionListChanged ColumnPermissionListChanged;
+
+	/** Delegate that broadcasts when creating actor browser columns. */
+	FOnCreateActorBrowserColumns CreateActorBrowserColumnsDelegate;
 	
 	TMap< FName, FSceneOutlinerFactory> CustomOutlinerFactories;
 
